Post Process Material Chain Graph: Added the plugin that allows users to execute Post process materials on textures other than scene color and write into a temporary buffer.

diff --git a/Engine/Plugins/Experimental/PostProcessMaterialChainGraph/PostProcessMaterialChainGraph.uplugin b/Engine/Plugins/Experimental/PostProcessMaterialChainGraph/PostProcessMaterialChainGraph.uplugin
new file mode 100644
--- /dev/null
+++ b/Engine/Plugins/Experimental/PostProcessMaterialChainGraph/PostProcessMaterialChainGraph.uplugin
@@ -0,0 +1,35 @@
+{
+	"FileVersion": 3,
+	"Version": 1,
+	"VersionName": "1.0",
+	"FriendlyName": "Post Process Material Chain Graph",
+	"Description": "Post Process Material Chain Graph allows users to stack post process materials and render those into render targets separate from Scene Color.\nThis can operate on textures other than scene color without writing those into scene color.",
+	"Category": "Other",
+	"CreatedBy": "Epic Games",
+	"CreatedByURL": "epicgames.com",
+	"DocsURL": "",
+	"MarketplaceURL": "",
+	"SupportURL": "",
+	"CanContainContent": true,
+	"IsBetaVersion": false,
+	"IsExperimentalVersion": true,
+	"Installed": false,
+	"Modules": [
+		{
+			"Name": "PPMChainGraph",
+			"Type": "Runtime",
+			"LoadingPhase": "PostConfigInit"
+		},
+		{
+			"Name": "PPMChainGraphEditor",
+			"Type": "Editor",
+			"LoadingPhase": "PreDefault",
+			"PlatformAllowList": [
+				"Win64",
+				"Linux",
+				"Mac"
+			]
+		}
+
+	]
+}
\ No newline at end of file
diff --git a/Engine/Plugins/Experimental/PostProcessMaterialChainGraph/Source/PostProcessMaterialChainGraph/PPMChainGraph.Build.cs b/Engine/Plugins/Experimental/PostProcessMaterialChainGraph/Source/PostProcessMaterialChainGraph/PPMChainGraph.Build.cs
new file mode 100644
--- /dev/null
+++ b/Engine/Plugins/Experimental/PostProcessMaterialChainGraph/Source/PostProcessMaterialChainGraph/PPMChainGraph.Build.cs
@@ -0,0 +1,33 @@
+// Copyright Epic Games, Inc. All Rights Reserved.
+
+using System.IO;
+using UnrealBuildTool;
+
+public class PPMChainGraph : ModuleRules
+{
+	public PPMChainGraph(ReadOnlyTargetRules Target) : base(Target)
+	{
+		PCHUsage = ModuleRules.PCHUsageMode.UseExplicitOrSharedPCHs;
+				
+		PrivateDependencyModuleNames.AddRange(
+			new string[]
+			{
+				"Core",
+				"CoreUObject",
+				"Engine",
+				"RHI",
+				"Projects",
+				"Renderer",
+				"RenderCore",
+				"Slate",
+				"SlateCore",
+
+			}
+		);
+
+		if (Target.Type == TargetType.Editor)
+		{
+			PrivateDependencyModuleNames.Add("UnrealEd");
+		}
+	}
+}
diff --git a/Engine/Plugins/Experimental/PostProcessMaterialChainGraph/Source/PostProcessMaterialChainGraph/Private/PPMChainGraph.cpp b/Engine/Plugins/Experimental/PostProcessMaterialChainGraph/Source/PostProcessMaterialChainGraph/Private/PPMChainGraph.cpp
new file mode 100644
--- /dev/null
+++ b/Engine/Plugins/Experimental/PostProcessMaterialChainGraph/Source/PostProcessMaterialChainGraph/Private/PPMChainGraph.cpp
@@ -0,0 +1,10 @@
+// Copyright Epic Games, Inc. All Rights Reserved.
+
+#include "PPMChainGraph.h"
+
+
+UPPMChainGraph::UPPMChainGraph(const FObjectInitializer& ObjectInitializer)
+	: Super(ObjectInitializer)
+{
+}
+
diff --git a/Engine/Plugins/Experimental/PostProcessMaterialChainGraph/Source/PostProcessMaterialChainGraph/Private/PPMChainGraphActor.cpp b/Engine/Plugins/Experimental/PostProcessMaterialChainGraph/Source/PostProcessMaterialChainGraph/Private/PPMChainGraphActor.cpp
new file mode 100644
--- /dev/null
+++ b/Engine/Plugins/Experimental/PostProcessMaterialChainGraph/Source/PostProcessMaterialChainGraph/Private/PPMChainGraphActor.cpp
@@ -0,0 +1,12 @@
+// Copyright Epic Games, Inc. All Rights Reserved.
+
+#include "PPMChainGraphActor.h"
+#include "PPMChainGraphComponent.h"
+
+APPMChainGraphActor::APPMChainGraphActor(const FObjectInitializer& ObjectInitializer)
+	: Super(ObjectInitializer)
+{
+	RootComponent = CreateDefaultSubobject<USceneComponent>(TEXT("Root"));
+
+	PPMChainGraphExecutorComponent = CreateDefaultSubobject<UPPMChainGraphExecutorComponent>(TEXT("PPMChainGraphExecutorComponent0"));
+}
diff --git a/Engine/Plugins/Experimental/PostProcessMaterialChainGraph/Source/PostProcessMaterialChainGraph/Private/PPMChainGraphComponent.cpp b/Engine/Plugins/Experimental/PostProcessMaterialChainGraph/Source/PostProcessMaterialChainGraph/Private/PPMChainGraphComponent.cpp
new file mode 100644
--- /dev/null
+++ b/Engine/Plugins/Experimental/PostProcessMaterialChainGraph/Source/PostProcessMaterialChainGraph/Private/PPMChainGraphComponent.cpp
@@ -0,0 +1,281 @@
+// Copyright Epic Games, Inc. All Rights Reserved.
+
+#include "PPMChainGraphComponent.h"
+#include "Camera/CameraActor.h"
+#include "Materials/Material.h"
+#include "MaterialDomain.h"
+#include "Framework/Notifications/NotificationManager.h"
+#include "PPMChainGraph.h"
+#include "PPMChainGraphWorldSubsystem.h"
+#include "RenderingThread.h"
+#include "Widgets/Notifications/SNotificationList.h"
+#include "Engine/Texture2D.h"
+#include "Engine/World.h"
+
+#define LOCTEXT_NAMESPACE "PPMChainGraphComponent"
+DEFINE_LOG_CATEGORY(LogPPMChainGraph);
+
+#define CAMERA_INVALID_PARENT_WARNING "Post Process Material Chain Graph Executor Component requires to be parented to Camera Actor. This component {0} on Actor \"{1}\" will be disabled until it is parented to Camera Actor."
+#define IS_STREAMED_TEXTURE_WARNING "The following external texture is streamed. {0} . \nThis may cause unexpected results. It is suggested to set NeverStream in texture properties."
+#define INVALID_MATERIAL_DOMAIN "The following material's domain isn't set to Post Processing. {0} . \nPost Process Material Chain Graph only supports Post Process Materials."
+#define RENDER_TARGET_ID_IS_IN_USE "Render target name is already in use. {0} . \nThe pass will not be executed until a valid name is provided."
+#define INPUT_IS_USED_AS_RENDERTARGET "Input is used as a render target. {0} . \nUsing the same input as a render target in the same pass is not allowed. Pass will not be executed."
+
+namespace
+{
+	struct FPPMChainGraphAggregatedWarnings
+	{
+		struct FWarnings
+		{
+			TSet<FString> TextureRequiresStreamingWarningList;
+			TSet<FString> InvalidMaterialDomainWarningList;
+			TSet<FString> TextureIdIsInUseWarningList;
+			TSet<FString> InputIsUsedAsAnOutputWarningList;
+		};
+		FWarnings CurrentFrame;
+		FWarnings LastFrame;
+	};
+}
+
+UPPMChainGraphExecutorComponent::UPPMChainGraphExecutorComponent(const FObjectInitializer& ObjectInitializer)
+	: Super(ObjectInitializer)
+	, AggregatedWarnings(MakeShared<FPPMChainGraphAggregatedWarnings>())
+{
+	PrimaryComponentTick.bCanEverTick = true;
+	bTickInEditor = true;
+	bAutoActivate = true;
+}
+
+void UPPMChainGraphExecutorComponent::TickComponent(float DeltaTime, enum ELevelTick TickType, FActorComponentTickFunction* ThisTickFunction)
+{
+	Super::TickComponent(DeltaTime, TickType, ThisTickFunction);
+	if (const UWorld* World = GetWorld())
+	{
+		if (!PPMChainGraphSubsystem.IsValid())
+		{
+			PPMChainGraphSubsystem = World->GetSubsystem<UPPMChainGraphWorldSubsystem>();
+		}
+
+		// We attempt to add every frame and remove it only on destroy to avoid complications in handling cases such as 
+		// undo/redo and other transactions. This way this object becomes invalid only when it is no longer ticking.
+		if (PPMChainGraphSubsystem.IsValid())
+		{
+			PPMChainGraphSubsystem->AddPPMChainGraphComponent(this);
+		}
+	}
+
+	// Validating graphs and preparing them for the use on render thread.
+	TransferState();
+
+	// Display aggregated warnings.
+	ProcessWarnings();
+};
+
+void UPPMChainGraphExecutorComponent::BeginDestroy()
+{
+	if (PPMChainGraphSubsystem.IsValid())
+	{
+		PPMChainGraphSubsystem->RemovePPMChainGraphComponent(this);
+	}
+
+	Super::BeginDestroy();
+}
+
+void UPPMChainGraphExecutorComponent::TransferState()
+{
+	TRACE_CPUPROFILER_EVENT_SCOPE_TEXT(*FString::Printf(TEXT("PPMChainGraph.TransferState %s"), *GetName()));
+	TMap<EPPMChainGraphExecutionLocation, TArray<TSharedPtr<FPPMChainGraphProxy>>> TempChainGraphRenderProxies;
+	FPPMChainGraphAggregatedWarnings* AggregatedWarningsCast = static_cast<FPPMChainGraphAggregatedWarnings*>(AggregatedWarnings.Get());
+
+	for (TWeakObjectPtr<UPPMChainGraph> Graph : PPMChainGraphs)
+	{
+		if (!Graph.IsValid())
+		{
+			continue;
+		}
+
+		// Ignore the graph if there are no passes.
+		if (Graph->Passes.IsEmpty())
+		{
+			continue;
+		}
+		TSharedPtr<FPPMChainGraphProxy> GraphRenderProxy = MakeShared<FPPMChainGraphProxy>();
+		GraphRenderProxy->PointOfExecution = Graph->PointOfExecution;
+		TArray<TSharedPtr<FPPMChainGraphProxy>>& ProxiesPerPointOfExecution = TempChainGraphRenderProxies.FindOrAdd(Graph->PointOfExecution);
+		ProxiesPerPointOfExecution.Add(GraphRenderProxy);
+
+		// Used to make sure that we don't have duplicate ids. 
+		TSet<FString> TextureIds;
+		for (TTuple<FString, TObjectPtr<UTexture2D>> KeyValue : Graph->ExternalTextures)
+		{
+			if (KeyValue.Key.IsEmpty() || !IsValid(KeyValue.Value))
+			{
+				continue;
+			}
+			TextureIds.Add(KeyValue.Key);
+			GraphRenderProxy->ExternalTextures.Add(KeyValue);
+
+			if (!KeyValue.Value->IsFullyStreamedIn())
+			{
+				FString TextureKey = FString::Format(TEXT("{0}, {1}"), { *KeyValue.Value.GetFullName(), *Graph->GetFullName()});
+				AggregatedWarningsCast->CurrentFrame.TextureRequiresStreamingWarningList.FindOrAdd(TextureKey);
+			}
+		}
+
+		// Validate each pass and create proxies.
+		for (const FPPMChainGraphPostProcessPass& Pass : Graph->Passes)
+		{
+			if (!Pass.bEnabled)
+			{
+				continue;
+			}
+			if (!Pass.PostProcessMaterial)
+			{
+				continue;
+			}
+
+			// Queue a warning about Post process material being invalid.
+			if (Pass.PostProcessMaterial->MaterialDomain != EMaterialDomain::MD_PostProcess)
+			{
+				FString MaterialIdString = FString::Format(TEXT("{0}, {1}"), { *Pass.PostProcessMaterial.GetFullName(), *Graph->GetFullName()});
+				AggregatedWarningsCast->CurrentFrame.InvalidMaterialDomainWarningList.FindOrAdd(MaterialIdString);
+				continue;
+			}
+
+			// Queue a warning about Render target Id already being used.
+			if (TextureIds.Contains(Pass.TemporaryRenderTargetId))
+			{
+				FString TextureId = FString::Format(TEXT("Render Target Name: {0}, {1}"), { *Pass.TemporaryRenderTargetId, *Graph->GetFullName() });
+				AggregatedWarningsCast->CurrentFrame.TextureIdIsInUseWarningList.FindOrAdd(TextureId);
+				continue;
+			}
+
+			// Ignore empty ids. 
+			if (Pass.Output == EPPMChainGraphOutput::PPMOutput_RenderTarget && Pass.TemporaryRenderTargetId.IsEmpty())
+			{
+				continue;
+			}
+
+			TextureIds.Add(Pass.TemporaryRenderTargetId);
+
+			TSharedPtr<FPPMChainGraphPostProcessPass> PassRenderProxy = MakeShared<FPPMChainGraphPostProcessPass>();
+			PassRenderProxy->PostProcessMaterial = Pass.PostProcessMaterial;
+			PassRenderProxy->TemporaryRenderTargetId = Pass.TemporaryRenderTargetId;
+			PassRenderProxy->Output = Pass.Output;
+
+			// We cannot write into the same render target that is used as input.
+			bool bInputReferencesOutput = false;
+
+			// Validate inputs.
+			for (TTuple<EPPMChainGraphPPMInputId, FString> KeyValue : Pass.Inputs)
+			{
+				if (KeyValue.Value == Pass.TemporaryRenderTargetId)
+				{
+					bInputReferencesOutput = true;
+
+					// Queue a warning about input being used as a render target.
+					FString TextureId = FString::Format(TEXT("{0}, {1}"), { *Pass.TemporaryRenderTargetId, *Graph->GetFullName() });
+					AggregatedWarningsCast->CurrentFrame.InputIsUsedAsAnOutputWarningList.FindOrAdd(TextureId);
+					break;
+				}
+				PassRenderProxy->Inputs.Add(KeyValue);
+			}
+
+			// Pass is invalid due to output writing into a texture that is used as an input in the same pass.
+			if (bInputReferencesOutput)
+			{
+				continue;
+			}
+
+			GraphRenderProxy->Passes.Add(MoveTemp(PassRenderProxy));
+		}
+	}
+
+	// Transfer the state on render thread, which is where it is going to be consumed.
+	ENQUEUE_RENDER_COMMAND(CopyPPMChainGraphProxy)([this, PPMChainGraphStateToCopy = MoveTemp(TempChainGraphRenderProxies)](FRHICommandListImmediate& RHICmdList)
+		{
+			FScopeLock ScopeLock(&StateTransferCriticalSection);
+			PPMChainGraphsRenderProxies = PPMChainGraphStateToCopy;
+		});
+}
+
+void UPPMChainGraphExecutorComponent::ProcessWarnings()
+{
+	FPPMChainGraphAggregatedWarnings* AggregatedWarningsCast = static_cast<FPPMChainGraphAggregatedWarnings*>(AggregatedWarnings.Get());
+	auto ShowToastFunc = [&](const FString& InText)
+	{
+		UE_LOG(LogPPMChainGraph, Warning, TEXT("%s"), *InText);
+
+#if WITH_EDITOR
+		FNotificationInfo Info(FText::FromString(InText));
+		Info.ExpireDuration = 5.0f;
+
+		FSlateNotificationManager::Get().AddNotification(Info);
+#endif
+	};
+	
+	auto DisplayWarning = [ShowToastFunc](const TSet<FString>& InLastFrameNotifications, const TSet<FString>& InCurrentFrameNotifications, const FString& InMessageFormat)
+	{
+		if (!InCurrentFrameNotifications.IsEmpty())
+		{
+			TSet<FString> DiffSet = InCurrentFrameNotifications.Difference(InLastFrameNotifications);
+			for (const FString& Notification : DiffSet)
+			{
+				if (InCurrentFrameNotifications.Contains(Notification))
+				{
+					ShowToastFunc(FString::Format(*InMessageFormat, {Notification}));
+				}
+			}
+		}
+	};
+
+	// Streaming texture warnings.
+	{
+		DisplayWarning(AggregatedWarningsCast->LastFrame.TextureRequiresStreamingWarningList
+			, AggregatedWarningsCast->CurrentFrame.TextureRequiresStreamingWarningList
+			, IS_STREAMED_TEXTURE_WARNING
+		);
+	}
+
+	// Invalid Material Domain warnings.
+	{
+		DisplayWarning(AggregatedWarningsCast->LastFrame.InvalidMaterialDomainWarningList
+			, AggregatedWarningsCast->CurrentFrame.InvalidMaterialDomainWarningList
+			, INVALID_MATERIAL_DOMAIN
+		);
+	}
+
+	// Render target Id is already in use warning.
+	{
+		DisplayWarning(AggregatedWarningsCast->LastFrame.TextureIdIsInUseWarningList
+			, AggregatedWarningsCast->CurrentFrame.TextureIdIsInUseWarningList
+			, RENDER_TARGET_ID_IS_IN_USE
+		);
+	}
+
+	// Render target Id is already in use warning.
+	{
+		DisplayWarning(AggregatedWarningsCast->LastFrame.InputIsUsedAsAnOutputWarningList
+			, AggregatedWarningsCast->CurrentFrame.InputIsUsedAsAnOutputWarningList
+			, INPUT_IS_USED_AS_RENDERTARGET
+		);
+	}
+	AggregatedWarningsCast->LastFrame = MoveTemp(AggregatedWarningsCast->CurrentFrame);
+	AggregatedWarningsCast->CurrentFrame = FPPMChainGraphAggregatedWarnings::FWarnings();
+
+}
+
+TArray<TSharedPtr<FPPMChainGraphProxy>> UPPMChainGraphExecutorComponent::GetChainGraphRenderProxies(EPPMChainGraphExecutionLocation InPointOfExecution)
+{
+	FScopeLock ScopeLock(&StateTransferCriticalSection);
+	if (PPMChainGraphsRenderProxies.Contains(InPointOfExecution))
+	{
+		return PPMChainGraphsRenderProxies[InPointOfExecution];
+	}
+	else
+	{
+		return TArray<TSharedPtr<FPPMChainGraphProxy>>();
+	}
+}
+
+#undef LOCTEXT_NAMESPACE 
\ No newline at end of file
diff --git a/Engine/Plugins/Experimental/PostProcessMaterialChainGraph/Source/PostProcessMaterialChainGraph/Private/PPMChainGraphModule.cpp b/Engine/Plugins/Experimental/PostProcessMaterialChainGraph/Source/PostProcessMaterialChainGraph/Private/PPMChainGraphModule.cpp
new file mode 100644
--- /dev/null
+++ b/Engine/Plugins/Experimental/PostProcessMaterialChainGraph/Source/PostProcessMaterialChainGraph/Private/PPMChainGraphModule.cpp
@@ -0,0 +1,13 @@
+// Copyright Epic Games, Inc. All Rights Reserved.
+
+#include "PPMChainGraphModule.h"
+
+void FPostProcessMaterialChainGraphModule::StartupModule()
+{
+}
+
+void FPostProcessMaterialChainGraphModule::ShutdownModule()
+{
+}
+	
+IMPLEMENT_MODULE(FPostProcessMaterialChainGraphModule, PostProcessMaterialChainGraph)
\ No newline at end of file
diff --git a/Engine/Plugins/Experimental/PostProcessMaterialChainGraph/Source/PostProcessMaterialChainGraph/Private/PPMChainGraphSceneViewExtension.cpp b/Engine/Plugins/Experimental/PostProcessMaterialChainGraph/Source/PostProcessMaterialChainGraph/Private/PPMChainGraphSceneViewExtension.cpp
new file mode 100644
--- /dev/null
+++ b/Engine/Plugins/Experimental/PostProcessMaterialChainGraph/Source/PostProcessMaterialChainGraph/Private/PPMChainGraphSceneViewExtension.cpp
@@ -0,0 +1,421 @@
+// Copyright Epic Games, Inc. All Rights Reserved.
+
+#include "PPMChainGraphSceneViewExtension.h"
+#include "Camera/CameraActor.h"
+#include "DataDrivenShaderPlatformInfo.h"
+#include "DynamicResolutionState.h"
+#include "FXRenderingUtils.h"
+#include "MaterialDomain.h"
+#include "Materials/Material.h"
+#include "Materials/MaterialRenderProxy.h"
+#include "PostProcess/DrawRectangle.h"
+#include "PostProcess/PostProcessMaterialInputs.h"
+#include "PostProcess/PostProcessInputs.h"
+#include "PPMChainGraphComponent.h"
+#include "PPMChainGraphWorldSubsystem.h"
+#include "PPMChainGraph.h"
+#include "RHI.h"
+#include "Runtime/Renderer/Private/SceneTextureParameters.h"
+#include "ScreenPass.h"
+#include "Engine/Texture2D.h"
+#include "TextureResource.h"
+
+DECLARE_GPU_STAT_NAMED(PPMChain, TEXT("PPMChainRender"));
+
+namespace
+{
+	void RenderPPMChainGraphProxies
+		( FRDGBuilder& GraphBuilder
+		, const FSceneView& View
+		, const FSceneViewFamily& ViewFamily
+		, const TArray<TSharedPtr<FPPMChainGraphProxy>>& ChainGraphProxy
+		, const FIntRect& PrimaryViewRect
+		, const FScreenPassRenderTarget& SceneColorRenderTarget
+		, FScreenPassRenderTarget& BackBufferRenderTarget
+		, FPostProcessMaterialInputs& PostProcessMaterialInputs
+		, EPPMChainGraphExecutionLocation PointOfExecution)
+	{
+		TRACE_CPUPROFILER_EVENT_SCOPE_TEXT(*FString::Printf(TEXT("PPMChainGraph.Render")));
+		SCOPED_GPU_STAT(GraphBuilder.RHICmdList, PPMChain);
+		FRHIDepthStencilState* DepthStencilState = FScreenPassPipelineState::FDefaultDepthStencilState::GetRHI();
+		
+		const FScreenPassTextureViewport RegionViewport(SceneColorRenderTarget.Texture, PrimaryViewRect);
+
+		FScreenPassRenderTarget MergedStencilRenderTarget;
+
+		bool bPPMApplied = false;
+
+		FRDGTextureDesc OutputDesc = SceneColorRenderTarget.Texture->Desc;
+		FLinearColor ClearColor(0., 0., 1., 0.);
+		OutputDesc.ClearValue = FClearValueBinding(ClearColor);
+		OutputDesc.Flags |= ETextureCreateFlags::RenderTargetable;
+
+		FScreenPassRenderTarget CopyBackTexture;
+		bool bCopyBackTexture = false;
+
+		FRHIBlendState* DefaultBlendState = FScreenPassPipelineState::FDefaultBlendState::GetRHI();
+		FGlobalShaderMap* GlobalShaderMap = GetGlobalShaderMap(GMaxRHIFeatureLevel);
+
+		for (const TSharedPtr<FPPMChainGraphProxy>& Graph : ChainGraphProxy)
+		{
+			if (!Graph.IsValid())
+			{
+				continue;
+			}
+
+			TMap<FString, FScreenPassRenderTarget> PassOutputs;
+			int32 PassId = 0;
+			for (const TSharedPtr<FPPMChainGraphPostProcessPass>& Pass : Graph->Passes)
+			{
+				if (!Pass->PostProcessMaterial || Pass->PostProcessMaterial->MaterialDomain != EMaterialDomain::MD_PostProcess)
+				{
+					continue;
+				}
+
+				bPPMApplied = true;
+
+				const FScreenPassTextureViewport OutputViewport(SceneColorRenderTarget);
+
+				if (Pass->Output == EPPMChainGraphOutput::PPMOutput_RenderTarget)
+				{
+					// Makes sure that the same render target won't be used as an input in the same pass. 
+					// There is an additional validation on Game thread.
+					if (PassOutputs.Contains(Pass->TemporaryRenderTargetId))
+					{
+						continue;
+					}
+					FString PassName = FString::Printf(TEXT("PPMChainGraphPassOutput_%d"), PassId++);
+					FRDGTexture* PassRDGTexture = GraphBuilder.CreateTexture(OutputDesc, *PassName);
+					FScreenPassRenderTarget PassRenderTarget = FScreenPassRenderTarget(PassRDGTexture, SceneColorRenderTarget.ViewRect, ERenderTargetLoadAction::EClear);
+					PassOutputs.Add(Pass->TemporaryRenderTargetId, PassRenderTarget);
+
+					PostProcessMaterialInputs.OverrideOutput = PassRenderTarget;
+					CopyBackTexture = PassRenderTarget;
+				}
+				else
+				{
+
+					PostProcessMaterialInputs.OverrideOutput = BackBufferRenderTarget;
+					CopyBackTexture = BackBufferRenderTarget;
+				}
+
+				bCopyBackTexture = true;
+
+				TStaticArray<FScreenPassTextureSlice, kPostProcessMaterialInputCountMax> PPMInputTextures;
+
+				// Mapping our custom inputs to PPM Inputs.
+				for (uint32 InputIndex = 0; InputIndex < kPostProcessMaterialInputCountMax; ++InputIndex)
+				{
+					bool bUseDefaultInput = true;
+					EPPMChainGraphPPMInputId InputId = (EPPMChainGraphPPMInputId)(InputIndex + 1);
+					if (Pass->Inputs.Contains(InputId))
+					{
+						FString InputMappedToId = Pass->Inputs[InputId];
+						// Need to exclude current pass.
+						if (PassOutputs.Contains(InputMappedToId))
+						{
+							PPMInputTextures[InputIndex] = FScreenPassTextureSlice::CreateFromScreenPassTexture(GraphBuilder, PassOutputs[InputMappedToId]);
+							bUseDefaultInput = false;
+						}
+						else if (Graph->ExternalTextures.Contains(InputMappedToId))
+						{
+							TWeakObjectPtr<UTexture2D> ExternalTexture = Graph->ExternalTextures[InputMappedToId];
+
+							if (ExternalTexture.IsValid())
+							{
+								FRHITexture* ExternalTextureResourceRHI = ExternalTexture->GetResource() ? ExternalTexture->GetResource()->GetTexture2DRHI() : nullptr;
+								if (ExternalTextureResourceRHI)
+								{
+									FRDGTextureRef ExternalTextureRDG = RegisterExternalTexture(GraphBuilder, ExternalTextureResourceRHI, *InputMappedToId);
+									FScreenPassTexture ExternalTextureInput(ExternalTextureRDG, FIntRect(0, 0, ExternalTexture->GetSizeX(), ExternalTexture->GetSizeY()));
+									PPMInputTextures[InputIndex] = FScreenPassTextureSlice::CreateFromScreenPassTexture(GraphBuilder, ExternalTextureInput);
+
+									bUseDefaultInput = false;
+								}
+							}
+						}
+					}
+					
+					if (bUseDefaultInput)
+					{
+						FScreenPassTextureSlice Input = PostProcessMaterialInputs.GetInput((EPostProcessMaterialInput)InputIndex);
+						PPMInputTextures[InputIndex] = Input;
+					}
+				}
+
+				PostProcessMaterialInputs.Textures = MoveTemp(PPMInputTextures);
+				AddPostProcessMaterialPass(GraphBuilder, View, PostProcessMaterialInputs, Pass->PostProcessMaterial);
+			}
+		}
+
+		// PrePostProcess is the only place where we need to copy back into Scene Color. In all other places we return back the backbuffer.
+		if (bCopyBackTexture && PointOfExecution == EPPMChainGraphExecutionLocation::PrePostProcess)
+		{
+			// Since we've rendered into the backbuffer already we have to use load flag instead.
+			CopyBackTexture.LoadAction = ERenderTargetLoadAction::ELoad;
+
+			TShaderMapRef<FScreenPassVS> VertexShader(GlobalShaderMap);
+			TShaderMapRef<FCopyRectPS> PixelShader(GlobalShaderMap);
+
+			FCopyRectPS::FParameters* PixelShaderParameters = GraphBuilder.AllocParameters<FCopyRectPS::FParameters>();
+			PixelShaderParameters->InputTexture = CopyBackTexture.Texture;
+			PixelShaderParameters->InputSampler = TStaticSamplerState<>::GetRHI();
+			PixelShaderParameters->RenderTargets[0] = SceneColorRenderTarget.GetRenderTargetBinding();
+
+			AddDrawScreenPass(
+				GraphBuilder,
+				RDG_EVENT_NAME("PPMChain_CopyViewport"),
+				View,
+				RegionViewport,
+				RegionViewport,
+				VertexShader,
+				PixelShader,
+				PixelShaderParameters
+			);
+		}
+	}
+}
+
+FPPMChainGraphSceneViewExtension::FPPMChainGraphSceneViewExtension(const FAutoRegister& AutoRegister, UPPMChainGraphWorldSubsystem* InWorldSubsystem)
+	: FSceneViewExtensionBase(AutoRegister)
+	, WorldSubsystem(InWorldSubsystem)
+{
+
+}
+
+void FPPMChainGraphSceneViewExtension::GatherChainGraphProxies
+(TArray<TSharedPtr<FPPMChainGraphProxy>>& OutChainGraphProxies
+	, const FSceneView& InView
+	, const FSceneViewFamily& InViewFamily
+	, EPPMChainGraphExecutionLocation InPointOfExecution)
+{
+	TSet<TWeakObjectPtr<UPPMChainGraphExecutorComponent>> PPMChainGraphComponents;
+	if (!WorldSubsystem.IsValid())
+	{
+		return;
+	}
+
+	{
+		FScopeLock ScopeLock(&WorldSubsystem->ComponentAccessCriticalSection);
+		PPMChainGraphComponents = WorldSubsystem->PPMChainGraphComponents;
+	}
+	for (TWeakObjectPtr<UPPMChainGraphExecutorComponent> PPMChainGraphComponent : PPMChainGraphComponents)
+	{
+		if (!PPMChainGraphComponent.IsValid())
+		{
+			continue;
+		}
+		if (PPMChainGraphComponent->PPMChainGraphs.Num() == 0)
+		{
+			continue;
+		}
+		if (
+			PPMChainGraphComponent->IsBeingDestroyed() ||
+			PPMChainGraphComponent->GetWorld() != InViewFamily.Scene->GetWorld()
+			)
+		{
+			continue;
+		}
+
+		if (PPMChainGraphComponent->CameraList.IsEmpty() && PPMChainGraphComponent->CameraViewHandlingMode == ECameraViewHandling::RenderOnlyInSelectedCameraViews)
+		{
+			continue;
+		}
+
+		bool bCameraOptionIsValid = PPMChainGraphComponent->CameraViewHandlingMode == ECameraViewHandling::IgnoreCameraViews ? true : false;
+		for (const TSoftObjectPtr<ACameraActor>& Camera : PPMChainGraphComponent->CameraList)
+		{
+			if (!Camera.IsValid())
+			{
+				continue;
+			}
+
+			if (PPMChainGraphComponent->CameraViewHandlingMode == ECameraViewHandling::IgnoreCameraViews)
+			{
+				if (Camera == InView.ViewActor)
+				{
+					bCameraOptionIsValid = false;
+					break;
+				}
+			}
+			else
+			{
+				if (Camera == InView.ViewActor)
+				{
+					bCameraOptionIsValid = true;
+					break;
+				}
+			}
+		}
+		if (!bCameraOptionIsValid)
+		{
+			continue;
+		}
+		OutChainGraphProxies.Append(PPMChainGraphComponent->GetChainGraphRenderProxies(InPointOfExecution));
+	}
+}
+
+void FPPMChainGraphSceneViewExtension::SubscribeToPostProcessingPass(EPostProcessingPass PassId, FAfterPassCallbackDelegateArray& InOutPassCallbacks, bool bIsPassEnabled)
+{
+	InOutPassCallbacks.Add(FAfterPassCallbackDelegate::CreateLambda([this, InPassId = PassId](FRDGBuilder& GraphBuilder, const FSceneView& View, const FPostProcessMaterialInputs& InInputs)
+		{
+			FPostProcessMaterialInputs InOutInputs = InInputs;
+			return FPPMChainGraphSceneViewExtension::AfterPostProcessPass_RenderThread(GraphBuilder, View, InOutInputs, InPassId);
+		}));
+}
+
+FScreenPassTexture FPPMChainGraphSceneViewExtension::AfterPostProcessPass_RenderThread(FRDGBuilder& GraphBuilder, const FSceneView& View, FPostProcessMaterialInputs& InOutInputs, EPostProcessingPass InCurrentPass)
+{
+	const FSceneViewFamily& ViewFamily = *View.Family;
+
+	TArray<TSharedPtr<FPPMChainGraphProxy>> ChainGraphProxies;
+
+	EPPMChainGraphExecutionLocation PointOfExecution = (EPPMChainGraphExecutionLocation)((int)InCurrentPass + 1);
+	GatherChainGraphProxies(ChainGraphProxies, View, ViewFamily, PointOfExecution);
+	FScreenPassRenderTarget OverrideOutput = InOutInputs.OverrideOutput;
+	if (ChainGraphProxies.IsEmpty())
+	{
+		// Don't need to modify anything, just return the untouched scene color texture back to post processing.
+		return InOutInputs.ReturnUntouchedSceneColorForPostProcessing(GraphBuilder);
+	}
+
+	const FIntRect PrimaryViewRect = UE::FXRenderingUtils::GetRawViewRectUnsafe(View);
+	const FScreenPassTexture SceneColor = FScreenPassTexture::CopyFromSlice(GraphBuilder, InOutInputs.GetInput(EPostProcessMaterialInput::SceneColor));
+
+	check(SceneColor.IsValid());
+
+	// Reusing the same output description for our back buffer as SceneColor
+	FRDGTextureDesc OutputDesc = SceneColor.Texture->Desc;
+	OutputDesc.Format = PF_FloatRGBA;
+	FLinearColor ClearColor(0., 0., 0., 0.);
+	OutputDesc.ClearValue = FClearValueBinding(ClearColor);
+	OutputDesc.Flags |= ETextureCreateFlags::RenderTargetable;
+
+	FScreenPassRenderTarget BackBufferRenderTarget = InOutInputs.OverrideOutput;
+
+	// If the override output is provided, it means that this is the last pass in post processing.
+	if (!BackBufferRenderTarget.IsValid())
+	{
+		FRDGTexture* BackBufferRenderTargetTexture = GraphBuilder.CreateTexture(OutputDesc, TEXT("BackBufferRenderTargetTexture"));
+		BackBufferRenderTarget = FScreenPassRenderTarget(BackBufferRenderTargetTexture, SceneColor.ViewRect, ERenderTargetLoadAction::EClear);
+	}
+
+	FScreenPassRenderTarget SceneColorRenderTarget(SceneColor, ERenderTargetLoadAction::ELoad);
+
+	RenderPPMChainGraphProxies
+	(
+		GraphBuilder
+		, View
+		, ViewFamily
+		, ChainGraphProxies
+		, PrimaryViewRect
+		, SceneColorRenderTarget
+		, BackBufferRenderTarget
+		, InOutInputs
+		, PointOfExecution
+	);
+
+	FRDGViewableResource* Resource = BackBufferRenderTarget.Texture;
+	if (Resource->HasBeenProduced())
+	{
+		return MoveTemp(BackBufferRenderTarget);
+	}
+	else
+	{
+		InOutInputs.OverrideOutput = OverrideOutput;
+		return InOutInputs.ReturnUntouchedSceneColorForPostProcessing(GraphBuilder);
+	}
+}
+
+void FPPMChainGraphSceneViewExtension::PrePostProcessPass_RenderThread(FRDGBuilder& GraphBuilder, const FSceneView& View, const FPostProcessingInputs& Inputs)
+{
+	const FSceneViewFamily& ViewFamily = *View.Family;
+
+	TArray<TSharedPtr<FPPMChainGraphProxy>> ChainGraphProxies;
+	GatherChainGraphProxies(ChainGraphProxies, View, ViewFamily, EPPMChainGraphExecutionLocation::PrePostProcess);
+
+	if (ChainGraphProxies.IsEmpty())
+	{
+		// Don't need to modify anything, just return the untouched scene color texture back to post processing.
+		return;
+	}
+
+	// We need to make sure to take Windows and Scene scale into account.
+	float ScreenPercentage = ViewFamily.SecondaryViewFraction;
+
+	if (ViewFamily.GetScreenPercentageInterface())
+	{
+		DynamicRenderScaling::TMap<float> UpperBounds = ViewFamily.GetScreenPercentageInterface()->GetResolutionFractionsUpperBound();
+		ScreenPercentage *= UpperBounds[GDynamicPrimaryResolutionFraction];
+	}
+
+	const FIntRect PrimaryViewRect = UE::FXRenderingUtils::GetRawViewRectUnsafe(View);
+	FScreenPassTexture SceneColor((*Inputs.SceneTextures)->SceneColorTexture, PrimaryViewRect);
+	check(SceneColor.IsValid());
+
+	// Reusing the same output description for our back buffer as SceneColor
+	FRDGTextureDesc OutputDesc = SceneColor.Texture->Desc;
+
+	OutputDesc.Format = PF_FloatRGBA;
+	FLinearColor ClearColor(0., 0., 0., 0.);
+	OutputDesc.ClearValue = FClearValueBinding(ClearColor);
+	OutputDesc.Flags |= ETextureCreateFlags::RenderTargetable;
+
+	FScreenPassRenderTarget SceneColorRenderTarget(SceneColor, ERenderTargetLoadAction::ELoad);
+	FRDGTexture* BackBufferRenderTargetTexture = GraphBuilder.CreateTexture(OutputDesc, TEXT("BackBufferRenderTargetTexture"));
+	FScreenPassRenderTarget BackBufferRenderTarget = FScreenPassRenderTarget(BackBufferRenderTargetTexture, SceneColor.ViewRect, ERenderTargetLoadAction::EClear);
+
+	const auto GetPostProcessMaterialInputs = [&](FScreenPassTexture InSceneColor)
+	{
+		FPostProcessMaterialInputs PostProcessMaterialInputs;
+
+		PostProcessMaterialInputs.SetInput(GraphBuilder, EPostProcessMaterialInput::SceneColor, InSceneColor);
+		const FScreenPassTexture Velocity((*Inputs.SceneTextures)->GBufferVelocityTexture, PrimaryViewRect);
+		FTranslucencyPassResources PostDOFTranslucencyResources = Inputs.TranslucencyViewResourcesMap.Get(ETranslucencyPass::TPT_TranslucencyAfterDOF);
+
+		FIntRect ViewRect{ 0, 0, 1, 1 };
+
+		if (Inputs.PathTracingResources.bPostProcessEnabled)
+		{
+			const FPathTracingResources& PathTracingResources = Inputs.PathTracingResources;
+
+			ViewRect = InSceneColor.ViewRect;
+			PostProcessMaterialInputs.SetPathTracingInput(EPathTracingPostProcessMaterialInput::Radiance, FScreenPassTexture(PathTracingResources.Radiance, ViewRect));
+			PostProcessMaterialInputs.SetPathTracingInput(EPathTracingPostProcessMaterialInput::DenoisedRadiance, FScreenPassTexture(PathTracingResources.DenoisedRadiance, ViewRect));
+			PostProcessMaterialInputs.SetPathTracingInput(EPathTracingPostProcessMaterialInput::Albedo, FScreenPassTexture(PathTracingResources.Albedo, ViewRect));
+			PostProcessMaterialInputs.SetPathTracingInput(EPathTracingPostProcessMaterialInput::Normal, FScreenPassTexture(PathTracingResources.Normal, ViewRect));
+			PostProcessMaterialInputs.SetPathTracingInput(EPathTracingPostProcessMaterialInput::Variance, FScreenPassTexture(PathTracingResources.Variance, ViewRect));
+		}
+
+		if (PostDOFTranslucencyResources.IsValid())
+		{
+			ViewRect = PostDOFTranslucencyResources.ViewRect;
+		}
+
+		PostProcessMaterialInputs.SetInput(GraphBuilder, EPostProcessMaterialInput::SeparateTranslucency, FScreenPassTexture(PostDOFTranslucencyResources.GetColorForRead(GraphBuilder), ViewRect));
+		PostProcessMaterialInputs.SetInput(GraphBuilder, EPostProcessMaterialInput::Velocity, Velocity);
+		PostProcessMaterialInputs.SceneTextures = GetSceneTextureShaderParameters(Inputs.SceneTextures);
+		const FScreenPassTexture CustomDepth(Inputs.CustomDepthTexture, PrimaryViewRect);
+		PostProcessMaterialInputs.CustomDepthTexture = CustomDepth.Texture;
+		PostProcessMaterialInputs.bManualStencilTest = Inputs.bSeparateCustomStencil;
+		// This is not available to Scene View Extensions yet.
+		//PostProcessMaterialInputs.SceneWithoutWaterTextures = &SceneWithoutWaterTextures;
+
+		return PostProcessMaterialInputs;
+	};
+
+	FPostProcessMaterialInputs PPMInputs = GetPostProcessMaterialInputs(SceneColor);
+
+	RenderPPMChainGraphProxies
+	(
+		GraphBuilder
+		, View
+		, ViewFamily
+		, ChainGraphProxies
+		, PrimaryViewRect
+		, SceneColorRenderTarget
+		, BackBufferRenderTarget
+		, PPMInputs
+		, EPPMChainGraphExecutionLocation::PrePostProcess
+	);
+}
diff --git a/Engine/Plugins/Experimental/PostProcessMaterialChainGraph/Source/PostProcessMaterialChainGraph/Private/PPMChainGraphSceneViewExtension.h b/Engine/Plugins/Experimental/PostProcessMaterialChainGraph/Source/PostProcessMaterialChainGraph/Private/PPMChainGraphSceneViewExtension.h
new file mode 100644
--- /dev/null
+++ b/Engine/Plugins/Experimental/PostProcessMaterialChainGraph/Source/PostProcessMaterialChainGraph/Private/PPMChainGraphSceneViewExtension.h
@@ -0,0 +1,31 @@
+// Copyright Epic Games, Inc. All Rights Reserved.
+
+#pragma once
+#include "SceneViewExtension.h"
+#include "PPMChainGraph.h"
+class UPPMChainGraphWorldSubsystem;
+struct FPPMChainGraphProxy;
+
+/**
+* Scene View Extension responsible for going through all PPM Chain Graph Components and rendering graphs into scene color.
+*/
+class FPPMChainGraphSceneViewExtension : public FSceneViewExtensionBase
+{
+public:
+	FPPMChainGraphSceneViewExtension(const FAutoRegister& AutoRegister, UPPMChainGraphWorldSubsystem* InWorldSubsystem);
+
+	//~ Begin FSceneViewExtensionBase Interface
+	virtual void SubscribeToPostProcessingPass(EPostProcessingPass PassId, FAfterPassCallbackDelegateArray& InOutPassCallbacks, bool bIsPassEnabled) override;
+
+	virtual void SetupViewFamily(FSceneViewFamily& InViewFamily) override {};
+	virtual void SetupView(FSceneViewFamily& InViewFamily, FSceneView& InView) override {};
+	virtual void BeginRenderViewFamily(FSceneViewFamily& InViewFamily) override {};
+	virtual void PrePostProcessPass_RenderThread(FRDGBuilder& GraphBuilder, const FSceneView& View, const FPostProcessingInputs& Inputs) override;
+	FScreenPassTexture AfterPostProcessPass_RenderThread(FRDGBuilder& GraphBuilder, const FSceneView& View, FPostProcessMaterialInputs& InOutInputs, EPostProcessingPass InCurrentPass);
+
+private:
+	void GatherChainGraphProxies(TArray<TSharedPtr<FPPMChainGraphProxy>>& OutChainGraphProxies, const FSceneView& InView, const FSceneViewFamily& InViewFamily, EPPMChainGraphExecutionLocation InPointOfExecution);
+
+private:
+	TWeakObjectPtr<UPPMChainGraphWorldSubsystem> WorldSubsystem;
+};
\ No newline at end of file
diff --git a/Engine/Plugins/Experimental/PostProcessMaterialChainGraph/Source/PostProcessMaterialChainGraph/Private/PPMChainGraphWorldSubsystem.cpp b/Engine/Plugins/Experimental/PostProcessMaterialChainGraph/Source/PostProcessMaterialChainGraph/Private/PPMChainGraphWorldSubsystem.cpp
new file mode 100644
--- /dev/null
+++ b/Engine/Plugins/Experimental/PostProcessMaterialChainGraph/Source/PostProcessMaterialChainGraph/Private/PPMChainGraphWorldSubsystem.cpp
@@ -0,0 +1,89 @@
+// Copyright Epic Games, Inc. All Rights Reserved.
+#include "PPMChainGraphWorldSubsystem.h"
+#include "PPMChainGraphSceneViewExtension.h"
+#include "PPMChainGraphComponent.h"
+
+#include "Engine/World.h"
+#include "Kismet/GameplayStatics.h"
+#include "EngineUtils.h"
+#include "SceneViewExtension.h"
+
+
+#if WITH_EDITOR
+#include "Editor.h"
+#endif
+
+#define LOCTEXT_NAMESPACE "PPMChainGraph"
+
+
+namespace
+{
+	bool IsComponentValid(UPPMChainGraphExecutorComponent* InPPMChainComponent, UWorld* CurrentWorld)
+	{
+		return InPPMChainComponent && InPPMChainComponent->GetWorld() == CurrentWorld;
+	}
+
+}
+
+void UPPMChainGraphWorldSubsystem::Initialize(FSubsystemCollectionBase& Collection)
+{
+	SceneViewExtension = FSceneViewExtensions::NewExtension<FPPMChainGraphSceneViewExtension>(this);
+	Super::Initialize(Collection);
+}
+
+void UPPMChainGraphWorldSubsystem::Deinitialize()
+{
+	// This has to be done on render thread as otherwise there could be a race condition when new level is loaded. 
+	ENQUEUE_RENDER_COMMAND(ReleaseSVE)([this](FRHICommandListImmediate& RHICmdList)
+	{
+		// Prevent this SVE from being gathered, in case it is kept alive by a strong reference somewhere else.
+		{
+			SceneViewExtension->IsActiveThisFrameFunctions.Empty();
+
+			FSceneViewExtensionIsActiveFunctor IsActiveFunctor;
+
+			IsActiveFunctor.IsActiveFunction = [](const ISceneViewExtension* InSceneViewExtension, const FSceneViewExtensionContext& InContext)
+			{
+				return TOptional<bool>(false);
+			};
+
+			SceneViewExtension->IsActiveThisFrameFunctions.Add(IsActiveFunctor);
+		}
+
+		SceneViewExtension.Reset();
+	});
+
+	// Finish all rendering commands before cleaning up actors.
+	FlushRenderingCommands();
+	{
+		FScopeLock ScopeLock(&ComponentAccessCriticalSection);
+		PPMChainGraphComponents.Reset();
+	}
+	Super::Deinitialize();
+
+}
+
+void UPPMChainGraphWorldSubsystem::Tick(float DeltaTime)
+{
+	Super::Tick(DeltaTime);
+}
+
+void UPPMChainGraphWorldSubsystem::AddPPMChainGraphComponent(TWeakObjectPtr<UPPMChainGraphExecutorComponent> InComponent)
+{
+	if (!PPMChainGraphComponents.Contains(InComponent))
+	{
+		FScopeLock ScopeLock(&ComponentAccessCriticalSection);
+		PPMChainGraphComponents.Add(InComponent);
+	}
+}
+
+void UPPMChainGraphWorldSubsystem::RemovePPMChainGraphComponent(TWeakObjectPtr<UPPMChainGraphExecutorComponent> InComponent)
+{
+	if (PPMChainGraphComponents.Contains(InComponent))
+	{
+		FScopeLock ScopeLock(&ComponentAccessCriticalSection);
+		PPMChainGraphComponents.Remove(InComponent);
+	}
+}
+
+#undef LOCTEXT_NAMESPACE
\ No newline at end of file
diff --git a/Engine/Plugins/Experimental/PostProcessMaterialChainGraph/Source/PostProcessMaterialChainGraph/Private/PPMChainGraphWorldSubsystem.h b/Engine/Plugins/Experimental/PostProcessMaterialChainGraph/Source/PostProcessMaterialChainGraph/Private/PPMChainGraphWorldSubsystem.h
new file mode 100644
--- /dev/null
+++ b/Engine/Plugins/Experimental/PostProcessMaterialChainGraph/Source/PostProcessMaterialChainGraph/Private/PPMChainGraphWorldSubsystem.h
@@ -0,0 +1,55 @@
+// Copyright Epic Games, Inc. All Rights Reserved.
+
+#pragma once
+#include "CoreMinimal.h"
+#include "UObject/ObjectMacros.h"
+#include "UObject/UObjectGlobals.h"
+#include "UObject/Object.h"
+#include "Subsystems/WorldSubsystem.h"
+#include "Engine/EngineBaseTypes.h"
+
+#if WITH_EDITOR
+#include "EditorUndoClient.h"
+#endif
+
+#include "PPMChainGraphWorldSubsystem.generated.h"
+
+class UPPMChainGraphExecutorComponent;
+
+/**
+ * World subsystem that is responsible for gathering and keeping track of PPM Chain Graph components.
+ */
+UCLASS()
+class UPPMChainGraphWorldSubsystem : public UTickableWorldSubsystem
+{
+	GENERATED_BODY()
+public:
+
+	// Subsystem Init/Deinit
+	virtual void Initialize(FSubsystemCollectionBase& Collection) override;
+	virtual void Deinitialize() override;
+
+	virtual bool IsTickableInEditor() const { return true; }
+	virtual void Tick(float DeltaTime) override;
+
+	// For profiling.
+	virtual TStatId GetStatId() const override { RETURN_QUICK_DECLARE_CYCLE_STAT(UPPMChainGraphWorldSubsystem, STATGROUP_Tickables); }
+public:
+	/** Adds a PPM Chain Graph component to the list of components to potentially be rendered. */
+	void AddPPMChainGraphComponent(TWeakObjectPtr<UPPMChainGraphExecutorComponent> InComponent);
+
+	/** Removes a PPM Chain Graph component. Typically this means that the component is removed from the world. */
+	void RemovePPMChainGraphComponent(TWeakObjectPtr<UPPMChainGraphExecutorComponent> InComponent);
+
+private:
+	TSharedPtr<class FPPMChainGraphSceneViewExtension, ESPMode::ThreadSafe> SceneViewExtension;
+
+	FCriticalSection ComponentAccessCriticalSection;
+
+	/** 
+	* All active chain graph components currently present in associated world.
+	*/
+	TSet<TWeakObjectPtr<UPPMChainGraphExecutorComponent>> PPMChainGraphComponents;
+public:
+	friend class FPPMChainGraphSceneViewExtension;
+};
\ No newline at end of file
diff --git a/Engine/Plugins/Experimental/PostProcessMaterialChainGraph/Source/PostProcessMaterialChainGraph/Public/PPMChainGraph.h b/Engine/Plugins/Experimental/PostProcessMaterialChainGraph/Source/PostProcessMaterialChainGraph/Public/PPMChainGraph.h
new file mode 100644
--- /dev/null
+++ b/Engine/Plugins/Experimental/PostProcessMaterialChainGraph/Source/PostProcessMaterialChainGraph/Public/PPMChainGraph.h
@@ -0,0 +1,130 @@
+// Copyright Epic Games, Inc. All Rights Reserved.
+
+#pragma once
+
+#include "UObject/Object.h"
+#include "CoreMinimal.h"
+#include "Engine/EngineTypes.h"
+
+#include "PPMChainGraph.generated.h"
+
+class UTextureRenderTarget2D;
+class UMaterial;
+class ACameraActor;
+
+UENUM()
+enum class EPPMChainGraphPPMInputId : uint8
+{
+	PPMInputMaping_Unassigned UMETA(DisplayName = "Unassigned"),
+	/** Map to equivalent input in Scene Textures node in the material. */
+	PPMInputMaping_0 UMETA(DisplayName = "PostProcessInput0"),
+	PPMInputMaping_1 UMETA(DisplayName = "PostProcessInput1"),
+	PPMInputMaping_2 UMETA(DisplayName = "PostProcessInput2"),
+	PPMInputMaping_3 UMETA(DisplayName = "PostProcessInput3"),
+	PPMInputMaping_4 UMETA(DisplayName = "PostProcessInput4"),
+};
+
+
+UENUM()
+enum class EPPMChainGraphOutput : uint8
+{
+	/**
+	* Write back into Scene Color.
+	*/
+	PPMOutput_SceneColor	UMETA(DisplayName = "Write Back into Scene Color"),
+
+	/** 
+	* Writes into a temporary render target to be used later in the chain. 
+	* Users need to provide the name for the render target for future identification.
+	*/
+	PPMOutput_RenderTarget	UMETA(DisplayName = "Temporary Render Target"),
+
+};
+
+
+UENUM()
+enum class EPPMChainGraphExecutionLocation : uint8
+{
+	// Default place of execution right before the rest of post processing starts.
+	PrePostProcess				UMETA(DisplayName = "Before Post Processing"),
+
+	// SSR Input isn't supported. 
+	//AfterSSRInput				UMETA(DisplayName = "After SSR Input"),
+	// 
+	// This is eventually mapped to equivalent execution location enum: ISceneViewExtension::EPostProcessingPass.
+	AfterMotionBlur				UMETA(DisplayName = "After Motion Blur"),
+	AfterToneMap				UMETA(DisplayName = "After Tonemap"),
+	AfterFXAA					UMETA(DisplayName = "After FXAA"),
+	AfterVisualizeDepthOfField	UMETA(DisplayName = "After Visualize Depth of Field"),
+};
+
+
+USTRUCT(Blueprintable)
+struct FPPMChainGraphPostProcessPass
+{
+	GENERATED_USTRUCT_BODY()
+
+	/** Is pass enabled. */
+	UPROPERTY(BlueprintReadWrite, EditAnywhere, Category = "Pass Settings")
+	bool bEnabled = true;
+
+	/** Inputs from previous passes. Map this to Scene Texture node in Post Process Material. */
+	UPROPERTY(BlueprintReadWrite, EditAnywhere, Category = "Post Process Pass Inputs")
+	TMap<EPPMChainGraphPPMInputId, FString> Inputs;
+	
+	/** 
+	* Where should this pass write to. By selecting Temporary Render Target as an option, 
+	* Users can avoid writing directly into Scene Color. For example this pass can operate
+	* on textures other than Scene Color to be used in one of the later passes.
+	* Keep in mind that the final pass will always write back into Scene Color.
+	*/
+	UPROPERTY(BlueprintReadWrite, EditAnywhere, Category = "Post Process Pass Output")
+	EPPMChainGraphOutput Output = EPPMChainGraphOutput::PPMOutput_SceneColor;
+
+	/** Use this to identify the Output of the current pass to be used later. */
+	UPROPERTY(BlueprintReadWrite, EditAnywhere, Category = "Post Process Pass Output", meta = (EditCondition = "Output == EPPMChainGraphOutput::PPMOutput_RenderTarget", EditConditionHides))
+	FString TemporaryRenderTargetId;
+
+	/**
+	* Which material should be executed during this pass.
+	*/
+	UPROPERTY(BlueprintReadWrite, EditAnywhere, Category = "Post Process Pass Material")
+	TObjectPtr<UMaterial> PostProcessMaterial;
+};
+
+
+/**
+* Post Process Material Chain Graph defines a collection of Post Process Material of passes that are executed one after another.
+* Individual passes can write into Scene Color or Temporary render target, but at the end of chain graph the combined result is 
+* always written into Scene Color.
+*/
+UCLASS(BlueprintType, meta = (DisplayName = "Post Process Material Chain Graph"))
+class PPMCHAINGRAPH_API UPPMChainGraph : public UObject
+{
+	GENERATED_BODY()
+
+public:
+
+	UPPMChainGraph(const FObjectInitializer& ObjectInitializer);
+public:
+
+	/** 
+	* Identifies at which point of post processing this graph is going to be executed. 
+	* These are different to the setting in Post Process Material.
+	*/
+	UPROPERTY(BlueprintReadWrite, EditAnywhere, Category = "Graph Execution Settings")
+	EPPMChainGraphExecutionLocation PointOfExecution;
+
+	/**
+	* External textures and rendertargets that can be mapped to Scene Texture Inputs.
+	*/
+	UPROPERTY(BlueprintReadWrite, EditAnywhere, Category = "External Inputs")
+	TMap<FString, TObjectPtr<UTexture2D>> ExternalTextures;
+
+	/**
+	* Post Process Material Passes. Each pass can write into the Scene Color or into a temporary render target that can
+	* be referenced in subsequent passes. At the end of all passes the result is always written into Scene Color.
+	*/
+	UPROPERTY(BlueprintReadWrite, EditAnywhere, Category = "Post Process Passes")
+	TArray<FPPMChainGraphPostProcessPass> Passes;
+};
\ No newline at end of file
diff --git a/Engine/Plugins/Experimental/PostProcessMaterialChainGraph/Source/PostProcessMaterialChainGraph/Public/PPMChainGraphActor.h b/Engine/Plugins/Experimental/PostProcessMaterialChainGraph/Source/PostProcessMaterialChainGraph/Public/PPMChainGraphActor.h
new file mode 100644
--- /dev/null
+++ b/Engine/Plugins/Experimental/PostProcessMaterialChainGraph/Source/PostProcessMaterialChainGraph/Public/PPMChainGraphActor.h
@@ -0,0 +1,18 @@
+// Copyright Epic Games, Inc. All Rights Reserved.
+
+#pragma once
+#include "CoreMinimal.h"
+#include "GameFramework/Actor.h"
+
+#include "PPMChainGraphActor.generated.h"
+
+class UPPMChainGraphExecutorComponent;
+
+UCLASS(Blueprintable, meta = (DisplayName = "Post Process Material Chain Graph Actor"))
+class PPMCHAINGRAPH_API APPMChainGraphActor : public AActor
+{
+	GENERATED_UCLASS_BODY()
+public:
+	UPROPERTY(Category = MediaPlate, VisibleAnywhere, BlueprintReadOnly)
+	TObjectPtr<UPPMChainGraphExecutorComponent> PPMChainGraphExecutorComponent;
+};
\ No newline at end of file
diff --git a/Engine/Plugins/Experimental/PostProcessMaterialChainGraph/Source/PostProcessMaterialChainGraph/Public/PPMChainGraphComponent.h b/Engine/Plugins/Experimental/PostProcessMaterialChainGraph/Source/PostProcessMaterialChainGraph/Public/PPMChainGraphComponent.h
new file mode 100644
--- /dev/null
+++ b/Engine/Plugins/Experimental/PostProcessMaterialChainGraph/Source/PostProcessMaterialChainGraph/Public/PPMChainGraphComponent.h
@@ -0,0 +1,80 @@
+// Copyright Epic Games, Inc. All Rights Reserved.
+
+#pragma once
+#include "CoreMinimal.h"
+#include "PPMChainGraph.h"
+#include "Components/SceneComponent.h"
+
+#include "PPMChainGraphComponent.generated.h"
+
+class UCameraComponent;
+class UPPMChainGraphWorldSubsystem;
+
+DECLARE_LOG_CATEGORY_EXTERN(LogPPMChainGraph, Log, All);
+
+struct FPPMChainGraphProxy
+{
+	TObjectPtr<ACameraActor> CameraOwner;
+
+	EPPMChainGraphExecutionLocation PointOfExecution;
+	TMap<FString, TObjectPtr<UTexture2D>> ExternalTextures;
+	TArray<TSharedPtr<FPPMChainGraphPostProcessPass>> Passes;
+};
+
+UENUM()
+enum class ECameraViewHandling : uint8
+{
+	IgnoreCameraViews UMETA(DisplayName = "Ignore Selected Camera Views"),
+	RenderOnlyInSelectedCameraViews UMETA(DisplayName = "Render Only In Selected Camera Views"),
+};
+
+UCLASS(hidecategories = (Transform, Collision, Object, Physics, SceneComponent, PostProcessVolume, Projection, Rendering, PlanarReflection), ClassGroup = Rendering, editinlinenew, meta = (BlueprintSpawnableComponent, DisplayName = "Post Process Material Chain Graph Executor Component"))
+class UPPMChainGraphExecutorComponent : public USceneComponent
+{
+	GENERATED_UCLASS_BODY()
+public:
+	virtual void TickComponent(float DeltaTime, ELevelTick TickType, FActorComponentTickFunction* ThisTickFunction) override;
+	virtual void BeginDestroy() override;
+
+public:
+
+	/** Get a copy of Render Proxies for rendering graphs. */
+	TArray<TSharedPtr<FPPMChainGraphProxy>> GetChainGraphRenderProxies(EPPMChainGraphExecutionLocation InPointOfExecution);
+
+public:
+	/** PPM Chain Graph Component can either exclude selected cameras from being rendered or render only in selected camera views. */
+	UPROPERTY(BlueprintReadWrite, EditAnywhere, Category = "Camera View Settings")
+	ECameraViewHandling CameraViewHandlingMode = ECameraViewHandling::IgnoreCameraViews;
+
+	/** 
+	* Depending on the selection in the option above the following cameras will either:
+	* 1. Be ignored and will not get PPM Chain Graph rendered into. 
+	* 2. Or rendered only in selected camera views.
+	*/
+	UPROPERTY(BlueprintReadWrite, EditAnywhere, Category = "Camera View Settings")
+	TArray<TSoftObjectPtr<ACameraActor>> CameraList;
+
+	/** A list of graphs to be executed. */
+	UPROPERTY(BlueprintReadWrite, EditAnywhere, Category = "Post Process Material Chain Graphs")
+	TArray<TObjectPtr<UPPMChainGraph>> PPMChainGraphs;
+
+private:
+	/** Transfers the render state on game thread. */
+	void TransferState();
+
+	/** Process aggregated warnings and notify the user. */
+	void ProcessWarnings();
+
+private:
+	// Critical section for tranfering state for rendering.
+	FCriticalSection StateTransferCriticalSection;
+
+	// Reference to the subsystem to add or remove this component to the list of Graphs to be considered.
+	TWeakObjectPtr<UPPMChainGraphWorldSubsystem> PPMChainGraphSubsystem;
+
+	// Populated on Game thread and used on Render thread for thread safety. Stored against the location of execution.
+	TMap<EPPMChainGraphExecutionLocation, TArray<TSharedPtr<FPPMChainGraphProxy>>> PPMChainGraphsRenderProxies;
+
+	// Collects warnings notifying users of invalid data in graphs or issues with ComponentAttachement.
+	TSharedPtr<void> AggregatedWarnings;
+};
\ No newline at end of file
diff --git a/Engine/Plugins/Experimental/PostProcessMaterialChainGraph/Source/PostProcessMaterialChainGraph/Public/PPMChainGraphModule.h b/Engine/Plugins/Experimental/PostProcessMaterialChainGraph/Source/PostProcessMaterialChainGraph/Public/PPMChainGraphModule.h
new file mode 100644
--- /dev/null
+++ b/Engine/Plugins/Experimental/PostProcessMaterialChainGraph/Source/PostProcessMaterialChainGraph/Public/PPMChainGraphModule.h
@@ -0,0 +1,15 @@
+// Copyright Epic Games, Inc. All Rights Reserved.
+
+#pragma once
+
+#include "CoreMinimal.h"
+#include "Modules/ModuleManager.h"
+
+class FPostProcessMaterialChainGraphModule : public IModuleInterface
+{
+public:
+
+	/** IModuleInterface implementation */
+	virtual void StartupModule() override;
+	virtual void ShutdownModule() override;
+};
diff --git a/Engine/Plugins/Experimental/PostProcessMaterialChainGraph/Source/PostProcessMaterialChainGraphEditor/PPMChainGraphEditor.Build.cs b/Engine/Plugins/Experimental/PostProcessMaterialChainGraph/Source/PostProcessMaterialChainGraphEditor/PPMChainGraphEditor.Build.cs
new file mode 100644
--- /dev/null
+++ b/Engine/Plugins/Experimental/PostProcessMaterialChainGraph/Source/PostProcessMaterialChainGraphEditor/PPMChainGraphEditor.Build.cs
@@ -0,0 +1,29 @@
+// Copyright Epic Games, Inc. All Rights Reserved.
+
+namespace UnrealBuildTool.Rules
+{
+	public class PPMChainGraphEditor : ModuleRules
+	{
+		public PPMChainGraphEditor(ReadOnlyTargetRules Target) : base(Target)
+		{
+			PrivateDependencyModuleNames.AddRange(
+				new string[] {
+					"AssetDefinition",
+					"Core",
+					"CoreUObject",
+					"EditorFramework",
+					"PPMChainGraph",
+					"UnrealEd",
+				});
+			
+			PrivateIncludePathModuleNames.AddRange(
+				new string[] {
+					
+				});
+
+			PrivateIncludePaths.AddRange(
+				new string[] {
+				});
+		}
+	}
+}
diff --git a/Engine/Plugins/Experimental/PostProcessMaterialChainGraph/Source/PostProcessMaterialChainGraphEditor/Private/AssetDefinition_PPMChainGraph.cpp b/Engine/Plugins/Experimental/PostProcessMaterialChainGraph/Source/PostProcessMaterialChainGraphEditor/Private/AssetDefinition_PPMChainGraph.cpp
new file mode 100644
--- /dev/null
+++ b/Engine/Plugins/Experimental/PostProcessMaterialChainGraph/Source/PostProcessMaterialChainGraphEditor/Private/AssetDefinition_PPMChainGraph.cpp
@@ -0,0 +1,32 @@
+// Copyright Epic Games, Inc. All Rights Reserved.
+
+#include "AssetDefinition_PPMChainGraph.h"
+
+#include "PPMChainGraph.h"
+
+#define LOCTEXT_NAMESPACE "AssetDefinition_PPMChainGraph"
+
+
+FText UAssetDefinition_PPMChainGraph::GetAssetDescription(const struct FAssetData& AssetData) const
+{
+	return FText::FromString("Post Process Material Chain Graph");
+}
+
+
+TSoftClassPtr<UObject> UAssetDefinition_PPMChainGraph::GetAssetClass() const
+{
+	return UPPMChainGraph::StaticClass();
+}
+
+FText UAssetDefinition_PPMChainGraph::GetAssetDisplayName() const
+{
+	return LOCTEXT("AssetDefinition_PPMChainGraph", "Post Process Material Chain Graph");
+}
+
+FLinearColor UAssetDefinition_PPMChainGraph::GetAssetColor() const
+{
+	return FLinearColor::White;
+}
+
+
+#undef LOCTEXT_NAMESPACE
diff --git a/Engine/Plugins/Experimental/PostProcessMaterialChainGraph/Source/PostProcessMaterialChainGraphEditor/Private/AssetDefinition_PPMChainGraph.h b/Engine/Plugins/Experimental/PostProcessMaterialChainGraph/Source/PostProcessMaterialChainGraphEditor/Private/AssetDefinition_PPMChainGraph.h
new file mode 100644
--- /dev/null
+++ b/Engine/Plugins/Experimental/PostProcessMaterialChainGraph/Source/PostProcessMaterialChainGraphEditor/Private/AssetDefinition_PPMChainGraph.h
@@ -0,0 +1,28 @@
+// Copyright Epic Games, Inc. All Rights Reserved.
+
+#pragma once
+
+#include "AssetDefinitionDefault.h"
+
+#include "AssetDefinition_PPMChainGraph.generated.h"
+
+UCLASS()
+class UAssetDefinition_PPMChainGraph : public UAssetDefinitionDefault
+{
+	GENERATED_BODY()
+
+public:
+	// Begin UAssetDefinition interface
+	virtual TConstArrayView<FAssetCategoryPath> GetAssetCategories() const override
+	{
+		static const auto Categories = { EAssetCategoryPaths::Misc };
+		return Categories;
+	}
+	virtual bool CanImport() const override { return false; }
+	
+	virtual FText GetAssetDisplayName() const override;
+	virtual FText GetAssetDescription(const FAssetData& AssetData) const override;
+	virtual TSoftClassPtr<UObject> GetAssetClass() const override;
+	virtual FLinearColor GetAssetColor() const override;
+	// End UAssetDefinition interface
+};
diff --git a/Engine/Plugins/Experimental/PostProcessMaterialChainGraph/Source/PostProcessMaterialChainGraphEditor/Private/Factories/ActorFactoryPPMChainGraphActor.cpp b/Engine/Plugins/Experimental/PostProcessMaterialChainGraph/Source/PostProcessMaterialChainGraphEditor/Private/Factories/ActorFactoryPPMChainGraphActor.cpp
new file mode 100644
--- /dev/null
+++ b/Engine/Plugins/Experimental/PostProcessMaterialChainGraph/Source/PostProcessMaterialChainGraphEditor/Private/Factories/ActorFactoryPPMChainGraphActor.cpp
@@ -0,0 +1,71 @@
+// Copyright Epic Games, Inc. All Rights Reserved.
+
+#include "ActorFactoryPPMChainGraphActor.h"
+#include "PPMChainGraph.h"
+#include "PPMChainGraphActor.h"
+#include "PPMChainGraphComponent.h"
+
+#include "Async/Async.h"
+
+
+#include UE_INLINE_GENERATED_CPP_BY_NAME(ActorFactoryPPMChainGraphActor)
+
+#define LOCTEXT_NAMESPACE "ActorFactoryPPMChainGraphActor"
+
+UActorFactoryPPMChainGraphActor::UActorFactoryPPMChainGraphActor(const FObjectInitializer& ObjectInitializer)
+	: Super(ObjectInitializer)
+{
+	DisplayName = LOCTEXT("PPMChainGraphActorDisplayName", "Post Process Material Chain Graph Actor");
+	NewActorClass = APPMChainGraphActor::StaticClass();
+}
+
+bool UActorFactoryPPMChainGraphActor::CanCreateActorFrom(const FAssetData& AssetData, FText& OutErrorMsg)
+{
+	if (AssetData.IsValid())
+	{
+		UClass* AssetClass = AssetData.GetClass();
+		if ((AssetClass != nullptr) && (AssetClass->IsChildOf(UPPMChainGraph::StaticClass())))
+		{
+			return true;
+		}
+		else
+		{
+			return false;
+		}
+	}
+	else
+	{
+		return true;
+	}
+}
+
+void UActorFactoryPPMChainGraphActor::PostSpawnActor(UObject* Asset, AActor* NewActor)
+{
+	Super::PostSpawnActor(Asset, NewActor);
+
+	SetUpActor(Asset, NewActor);
+}
+
+void UActorFactoryPPMChainGraphActor::SetUpActor(UObject* Asset, AActor* Actor)
+{
+	if (Actor != nullptr)
+	{
+		APPMChainGraphActor* ChainGraphActor = CastChecked<APPMChainGraphActor>(Actor);
+
+		UPPMChainGraph* ChainGraph = Cast<UPPMChainGraph>(Asset);
+
+		if ((ChainGraph != nullptr) && (ChainGraphActor->PPMChainGraphExecutorComponent != nullptr) && !ChainGraphActor->bIsEditorPreviewActor)
+		{
+			if (UPPMChainGraphExecutorComponent* PPMChainGraphComponent = GetValid(ChainGraphActor->PPMChainGraphExecutorComponent))
+			{
+				if (!PPMChainGraphComponent->PPMChainGraphs.Contains(ChainGraph))
+				{
+					PPMChainGraphComponent->PPMChainGraphs.Add(ChainGraph);
+				}
+			}
+		}
+	}
+}
+
+#undef LOCTEXT_NAMESPACE
+
diff --git a/Engine/Plugins/Experimental/PostProcessMaterialChainGraph/Source/PostProcessMaterialChainGraphEditor/Private/Factories/ActorFactoryPPMChainGraphActor.h b/Engine/Plugins/Experimental/PostProcessMaterialChainGraph/Source/PostProcessMaterialChainGraphEditor/Private/Factories/ActorFactoryPPMChainGraphActor.h
new file mode 100644
--- /dev/null
+++ b/Engine/Plugins/Experimental/PostProcessMaterialChainGraph/Source/PostProcessMaterialChainGraphEditor/Private/Factories/ActorFactoryPPMChainGraphActor.h
@@ -0,0 +1,25 @@
+// Copyright Epic Games, Inc. All Rights Reserved.
+
+#pragma once
+
+#include "CoreMinimal.h"
+#include "UObject/ObjectMacros.h"
+#include "ActorFactories/ActorFactory.h"
+#include "ActorFactoryPPMChainGraphActor.generated.h"
+
+UCLASS()
+class UActorFactoryPPMChainGraphActor : public UActorFactory
+{
+	GENERATED_UCLASS_BODY()
+
+	//~ Begin UActorFactory Interface
+	virtual bool CanCreateActorFrom(const FAssetData& AssetData, FText& OutErrorMsg) override;
+	virtual void PostSpawnActor(UObject* Asset, AActor* NewActor) override;
+	//~ End UActorFactory Interface
+
+private:
+	/**
+	 * Set up a PPM Chain Graph Actor.
+	 */
+	void SetUpActor(UObject* Asset, AActor* Actor);
+};
diff --git a/Engine/Plugins/Experimental/PostProcessMaterialChainGraph/Source/PostProcessMaterialChainGraphEditor/Private/Factories/PPMChainGraphFactoryNew.cpp b/Engine/Plugins/Experimental/PostProcessMaterialChainGraph/Source/PostProcessMaterialChainGraphEditor/Private/Factories/PPMChainGraphFactoryNew.cpp
new file mode 100644
--- /dev/null
+++ b/Engine/Plugins/Experimental/PostProcessMaterialChainGraph/Source/PostProcessMaterialChainGraphEditor/Private/Factories/PPMChainGraphFactoryNew.cpp
@@ -0,0 +1,36 @@
+// Copyright Epic Games, Inc. All Rights Reserved.
+
+#include "PPMChainGraphFactoryNew.h"
+
+#include "AssetTypeCategories.h"
+#include "PPMChainGraph.h"
+
+#include UE_INLINE_GENERATED_CPP_BY_NAME(PPMChainGraphFactoryNew)
+
+
+UPPMChainGraphFactoryNew::UPPMChainGraphFactoryNew(const FObjectInitializer& ObjectInitializer)
+	: Super(ObjectInitializer)
+{
+	SupportedClass = UPPMChainGraph::StaticClass();
+	bCreateNew = true;
+	bEditAfterNew = true;
+}
+
+
+UObject* UPPMChainGraphFactoryNew::FactoryCreateNew(UClass* InClass, UObject* InParent, FName InName, EObjectFlags Flags, UObject* Context, FFeedbackContext* Warn)
+{
+	return NewObject<UPPMChainGraph>(InParent, InClass, InName, Flags);
+}
+
+
+uint32 UPPMChainGraphFactoryNew::GetMenuCategories() const
+{
+	return EAssetTypeCategories::Misc;
+}
+
+
+bool UPPMChainGraphFactoryNew::ShouldShowInNewMenu() const
+{
+	return true;
+}
+
diff --git a/Engine/Plugins/Experimental/PostProcessMaterialChainGraph/Source/PostProcessMaterialChainGraphEditor/Private/Factories/PPMChainGraphFactoryNew.h b/Engine/Plugins/Experimental/PostProcessMaterialChainGraph/Source/PostProcessMaterialChainGraphEditor/Private/Factories/PPMChainGraphFactoryNew.h
new file mode 100644
--- /dev/null
+++ b/Engine/Plugins/Experimental/PostProcessMaterialChainGraph/Source/PostProcessMaterialChainGraphEditor/Private/Factories/PPMChainGraphFactoryNew.h
@@ -0,0 +1,22 @@
+// Copyright Epic Games, Inc. All Rights Reserved.
+
+#pragma once
+
+#include "Factories/Factory.h"
+
+#include "PPMChainGraphFactoryNew.generated.h"
+
+
+UCLASS(hidecategories=Object)
+class UPPMChainGraphFactoryNew : public UFactory
+{
+	GENERATED_UCLASS_BODY()
+
+public:
+
+	//~ UFactory Interface
+
+	virtual UObject* FactoryCreateNew(UClass* InClass, UObject* InParent, FName InName, EObjectFlags Flags, UObject* Context, FFeedbackContext* Warn) override;
+	virtual uint32 GetMenuCategories() const override;
+	virtual bool ShouldShowInNewMenu() const override;
+};
diff --git a/Engine/Plugins/Experimental/PostProcessMaterialChainGraph/Source/PostProcessMaterialChainGraphEditor/Private/PPMChainGraphEditorModule.cpp b/Engine/Plugins/Experimental/PostProcessMaterialChainGraph/Source/PostProcessMaterialChainGraphEditor/Private/PPMChainGraphEditorModule.cpp
new file mode 100644
--- /dev/null
+++ b/Engine/Plugins/Experimental/PostProcessMaterialChainGraph/Source/PostProcessMaterialChainGraphEditor/Private/PPMChainGraphEditorModule.cpp
@@ -0,0 +1,20 @@
+// Copyright Epic Games, Inc. All Rights Reserved.
+
+#include "PPMChainGraphEditorModule.h"
+#include "Modules/ModuleManager.h"
+
+
+void FPPMChainGraphEditorModule::StartupModule()
+{
+}
+
+void FPPMChainGraphEditorModule::ShutdownModule()
+{
+	
+}
+
+IMPLEMENT_MODULE(FPPMChainGraphEditorModule, PPMChainGraphEditor);
+DEFINE_LOG_CATEGORY(LogPPMChainGraphEditor);
+
+
+#undef LOCTEXT_NAMESPACE
diff --git a/Engine/Plugins/Experimental/PostProcessMaterialChainGraph/Source/PostProcessMaterialChainGraphEditor/Private/PPMChainGraphEditorModule.h b/Engine/Plugins/Experimental/PostProcessMaterialChainGraph/Source/PostProcessMaterialChainGraphEditor/Private/PPMChainGraphEditorModule.h
new file mode 100644
--- /dev/null
+++ b/Engine/Plugins/Experimental/PostProcessMaterialChainGraph/Source/PostProcessMaterialChainGraphEditor/Private/PPMChainGraphEditorModule.h
@@ -0,0 +1,20 @@
+// Copyright Epic Games, Inc. All Rights Reserved.
+
+#pragma once
+
+#include "IPPMChainGraphEditorModule.h"
+#include "CoreMinimal.h"
+
+
+class FPPMChainGraphEditorModule : public IPPMChainGraphEditorModule, public TSharedFromThis<FPPMChainGraphEditorModule>
+{
+public:
+
+	virtual ~FPPMChainGraphEditorModule() = default;
+
+	//~ Begin IModuleInterface interface
+	virtual void StartupModule() override;
+	virtual void ShutdownModule() override;
+	//~ End IModuleInterface interface
+};
+
diff --git a/Engine/Plugins/Experimental/PostProcessMaterialChainGraph/Source/PostProcessMaterialChainGraphEditor/Public/IPPMChainGraphEditorModule.h b/Engine/Plugins/Experimental/PostProcessMaterialChainGraph/Source/PostProcessMaterialChainGraphEditor/Public/IPPMChainGraphEditorModule.h
new file mode 100644
--- /dev/null
+++ b/Engine/Plugins/Experimental/PostProcessMaterialChainGraph/Source/PostProcessMaterialChainGraphEditor/Public/IPPMChainGraphEditorModule.h
@@ -0,0 +1,16 @@
+// Copyright Epic Games, Inc. All Rights Reserved.
+
+#pragma once
+
+#include "Logging/LogMacros.h"
+#include "Modules/ModuleInterface.h"
+#include "Templates/SharedPointer.h"
+
+
+DECLARE_LOG_CATEGORY_EXTERN(LogPPMChainGraphEditor, Log, All);
+
+class IPPMChainGraphEditorModule : public IModuleInterface
+{
+public:
+};
+
