Fix Found an '#undef LOCTEXT_NAMESPACE' without a corresponding '#define LOCTEXT_NAMESPACE'
#rnx

diff --git a/Engine/Plugins/Experimental/PostProcessMaterialChainGraph/Source/PostProcessMaterialChainGraphEditor/Private/PPMChainGraphEditorModule.cpp b/Engine/Plugins/Experimental/PostProcessMaterialChainGraph/Source/PostProcessMaterialChainGraphEditor/Private/PPMChainGraphEditorModule.cpp
--- a/Engine/Plugins/Experimental/PostProcessMaterialChainGraph/Source/PostProcessMaterialChainGraphEditor/Private/PPMChainGraphEditorModule.cpp
+++ b/Engine/Plugins/Experimental/PostProcessMaterialChainGraph/Source/PostProcessMaterialChainGraphEditor/Private/PPMChainGraphEditorModule.cpp
@@ -15,6 +15,3 @@ void FPPMChainGraphEditorModule::ShutdownModule()
 
 IMPLEMENT_MODULE(FPPMChainGraphEditorModule, PPMChainGraphEditor);
 DEFINE_LOG_CATEGORY(LogPPMChainGraphEditor);
-
-
-#undef LOCTEXT_NAMESPACE
