Post process material chain graph: Fixed a few issues and added a few improvements:

- PPMChain Graph actor no longer renders when it is hidden in editor.
- User is notified if external texture is virtual texture
- no longer crashes if virtual texture is used.
- PPM Chain Graph SVE doesn't subscribe to any after passes if there are no graphs that need to be executed at different stage of pipeline.
- added billboard component to PPM Chain Graph Actors
- fixed the issue with the black screen in nDisplay.

diff --git a/Engine/Plugins/Experimental/PostProcessMaterialChainGraph/Source/PostProcessMaterialChainGraph/Private/PPMChainGraphActor.cpp b/Engine/Plugins/Experimental/PostProcessMaterialChainGraph/Source/PostProcessMaterialChainGraph/Private/PPMChainGraphActor.cpp
--- a/Engine/Plugins/Experimental/PostProcessMaterialChainGraph/Source/PostProcessMaterialChainGraph/Private/PPMChainGraphActor.cpp
+++ b/Engine/Plugins/Experimental/PostProcessMaterialChainGraph/Source/PostProcessMaterialChainGraph/Private/PPMChainGraphActor.cpp
@@ -2,6 +2,9 @@
 
 #include "PPMChainGraphActor.h"
 #include "PPMChainGraphComponent.h"
+#include "Components/BillboardComponent.h"
+#include "UObject/ConstructorHelpers.h"
+#include "Engine/Texture2D.h"
 
 APPMChainGraphActor::APPMChainGraphActor(const FObjectInitializer& ObjectInitializer)
 	: Super(ObjectInitializer)
@@ -9,4 +12,48 @@ APPMChainGraphActor::APPMChainGraphActor(const FObjectInitializer& ObjectInitial
 	RootComponent = CreateDefaultSubobject<USceneComponent>(TEXT("Root"));
 
 	PPMChainGraphExecutorComponent = CreateDefaultSubobject<UPPMChainGraphExecutorComponent>(TEXT("PPMChainGraphExecutorComponent0"));
+
+
+#if WITH_METADATA
+	{
+		// Create billboard component
+		if (GIsEditor && !IsRunningCommandlet())
+		{
+			// Structure to hold one-time initialization
+
+			struct FConstructorStatics
+			{
+				ConstructorHelpers::FObjectFinderOptional<UTexture2D> SpriteTextureObject;
+				FName ID_PPMChainGraphActor;
+				FText NAME_PPMChainGraphActor;
+
+				FConstructorStatics()
+					: SpriteTextureObject(TEXT("/Engine/EditorResources/S_Actor"))
+					, ID_PPMChainGraphActor(TEXT("PPM Chain Graph Actor"))
+					, NAME_PPMChainGraphActor(NSLOCTEXT("SpriteCategory", "PPMChainGraphActor", "PPM Chain Graph Actor"))
+				{
+				}
+			};
+
+			static FConstructorStatics ConstructorStatics;
+
+			SpriteComponent = CreateEditorOnlyDefaultSubobject<UBillboardComponent>(TEXT("PPM Chain Graph Component Icon"));
+
+			if (SpriteComponent)
+			{
+				SpriteComponent->Sprite = ConstructorStatics.SpriteTextureObject.Get();
+				SpriteComponent->SpriteInfo.Category = ConstructorStatics.ID_PPMChainGraphActor;
+				SpriteComponent->SpriteInfo.DisplayName = ConstructorStatics.NAME_PPMChainGraphActor;
+				SpriteComponent->SetIsVisualizationComponent(true);
+				SpriteComponent->SetRelativeLocationAndRotation(FVector::ZeroVector, FRotator::ZeroRotator);
+				SpriteComponent->SetMobility(EComponentMobility::Movable);
+				SpriteComponent->bHiddenInGame = true;
+				SpriteComponent->bIsScreenSizeScaled = true;
+
+				SpriteComponent->AttachToComponent(RootComponent, FAttachmentTransformRules::KeepRelativeTransform);
+			}
+		}
+
+	}
+#endif 
 }
diff --git a/Engine/Plugins/Experimental/PostProcessMaterialChainGraph/Source/PostProcessMaterialChainGraph/Private/PPMChainGraphComponent.cpp b/Engine/Plugins/Experimental/PostProcessMaterialChainGraph/Source/PostProcessMaterialChainGraph/Private/PPMChainGraphComponent.cpp
--- a/Engine/Plugins/Experimental/PostProcessMaterialChainGraph/Source/PostProcessMaterialChainGraph/Private/PPMChainGraphComponent.cpp
+++ b/Engine/Plugins/Experimental/PostProcessMaterialChainGraph/Source/PostProcessMaterialChainGraph/Private/PPMChainGraphComponent.cpp
@@ -16,7 +16,7 @@
 DEFINE_LOG_CATEGORY(LogPPMChainGraph);
 
 #define CAMERA_INVALID_PARENT_WARNING "Post Process Material Chain Graph Executor Component requires to be parented to Camera Actor. This component {0} on Actor \"{1}\" will be disabled until it is parented to Camera Actor."
-#define IS_STREAMED_TEXTURE_WARNING "The following external texture is streamed. {0} . \nThis may cause unexpected results. It is suggested to set NeverStream in texture properties."
+#define IS_STREAMED_TEXTURE_WARNING "The following external texture is streamed. {0} . \nThis may cause unexpected results. It is suggested to set NeverStream in texture properties and making sure Virtual Texture streaming is disabled."
 #define INVALID_MATERIAL_DOMAIN "The following material's domain isn't set to Post Processing. {0} . \nPost Process Material Chain Graph only supports Post Process Materials."
 #define RENDER_TARGET_ID_IS_IN_USE "Render target name is already in use. {0} . \nThe pass will not be executed until a valid name is provided."
 #define INPUT_IS_USED_AS_RENDERTARGET "Input is used as a render target. {0} . \nUsing the same input as a render target in the same pass is not allowed. Pass will not be executed."
@@ -49,6 +49,28 @@ UPPMChainGraphExecutorComponent::UPPMChainGraphExecutorComponent(const FObjectIn
 void UPPMChainGraphExecutorComponent::TickComponent(float DeltaTime, enum ELevelTick TickType, FActorComponentTickFunction* ThisTickFunction)
 {
 	Super::TickComponent(DeltaTime, TickType, ThisTickFunction);
+
+	{
+		AActor* OwnerActor = GetOwner();
+		if (!IsValid(OwnerActor))
+		{
+			return;
+		}
+
+		bool bIsActiveLoc = IsValid(this)
+#if WITH_EDITOR
+			&& !GetOwner()->IsHiddenEd()
+#endif 
+			&& !GetOwner()->IsHidden();
+
+		if (!bIsActiveLoc)
+		{
+			FScopeLock ScopeLock(&StateTransferCriticalSection);
+			PPMChainGraphsRenderProxies.Empty();
+			return;
+		}
+	}
+
 	if (const UWorld* World = GetWorld())
 	{
 		if (!PPMChainGraphSubsystem.IsValid())
@@ -115,7 +137,7 @@ void UPPMChainGraphExecutorComponent::TransferState()
 			TextureIds.Add(KeyValue.Key);
 			GraphRenderProxy->ExternalTextures.Add(KeyValue);
 
-			if (!KeyValue.Value->IsFullyStreamedIn())
+			if (!KeyValue.Value->IsFullyStreamedIn() || KeyValue.Value->VirtualTextureStreaming)
 			{
 				FString TextureKey = FString::Format(TEXT("{0}, {1}"), { *KeyValue.Value.GetFullName(), *Graph->GetFullName()});
 				AggregatedWarningsCast->CurrentFrame.TextureRequiresStreamingWarningList.FindOrAdd(TextureKey);
@@ -192,11 +214,10 @@ void UPPMChainGraphExecutorComponent::TransferState()
 	}
 
 	// Transfer the state on render thread, which is where it is going to be consumed.
-	ENQUEUE_RENDER_COMMAND(CopyPPMChainGraphProxy)([this, PPMChainGraphStateToCopy = MoveTemp(TempChainGraphRenderProxies)](FRHICommandListImmediate& RHICmdList)
-		{
-			FScopeLock ScopeLock(&StateTransferCriticalSection);
-			PPMChainGraphsRenderProxies = PPMChainGraphStateToCopy;
-		});
+	{
+		FScopeLock ScopeLock(&StateTransferCriticalSection);
+		PPMChainGraphsRenderProxies = TempChainGraphRenderProxies;
+	};
 }
 
 void UPPMChainGraphExecutorComponent::ProcessWarnings()
@@ -278,4 +299,10 @@ TArray<TSharedPtr<FPPMChainGraphProxy>> UPPMChainGraphExecutorComponent::GetChai
 	}
 }
 
+bool UPPMChainGraphExecutorComponent::IsActiveDuringPass_GameThread(EPPMChainGraphExecutionLocation InPointOfExecution)
+{
+	check(IsInGameThread());
+	return PPMChainGraphsRenderProxies.Contains(InPointOfExecution);
+}
+
 #undef LOCTEXT_NAMESPACE 
\ No newline at end of file
diff --git a/Engine/Plugins/Experimental/PostProcessMaterialChainGraph/Source/PostProcessMaterialChainGraph/Private/PPMChainGraphSceneViewExtension.cpp b/Engine/Plugins/Experimental/PostProcessMaterialChainGraph/Source/PostProcessMaterialChainGraph/Private/PPMChainGraphSceneViewExtension.cpp
--- a/Engine/Plugins/Experimental/PostProcessMaterialChainGraph/Source/PostProcessMaterialChainGraph/Private/PPMChainGraphSceneViewExtension.cpp
+++ b/Engine/Plugins/Experimental/PostProcessMaterialChainGraph/Source/PostProcessMaterialChainGraph/Private/PPMChainGraphSceneViewExtension.cpp
@@ -84,8 +84,7 @@ namespace
 					{
 						continue;
 					}
-					FString PassName = FString::Printf(TEXT("PPMChainGraphPassOutput_%d"), PassId++);
-					FRDGTexture* PassRDGTexture = GraphBuilder.CreateTexture(OutputDesc, *PassName);
+					FRDGTexture* PassRDGTexture = GraphBuilder.CreateTexture(OutputDesc, *Pass->TemporaryRenderTargetId);
 					FScreenPassRenderTarget PassRenderTarget = FScreenPassRenderTarget(PassRDGTexture, SceneColorRenderTarget.ViewRect, ERenderTargetLoadAction::EClear);
 					PassOutputs.Add(Pass->TemporaryRenderTargetId, PassRenderTarget);
 
@@ -124,8 +123,19 @@ namespace
 							if (ExternalTexture.IsValid())
 							{
 								FRHITexture* ExternalTextureResourceRHI = ExternalTexture->GetResource() ? ExternalTexture->GetResource()->GetTexture2DRHI() : nullptr;
+								bool bIsShaderResource = false;
+
 								if (ExternalTextureResourceRHI)
 								{
+									const FRHITextureDesc& ExternalTexDesc = ExternalTextureResourceRHI->GetDesc();
+
+									// Virtual textures are not used as shader resources and users are notified about it.
+									bIsShaderResource = EnumHasAnyFlags(ExternalTexDesc.Flags, ETextureCreateFlags::ShaderResource);
+								}
+
+								if (ExternalTextureResourceRHI && bIsShaderResource)
+								{
+									ExternalTextureResourceRHI->GetDesc();
 									FRDGTextureRef ExternalTextureRDG = RegisterExternalTexture(GraphBuilder, ExternalTextureResourceRHI, *InputMappedToId);
 									FScreenPassTexture ExternalTextureInput(ExternalTextureRDG, FIntRect(0, 0, ExternalTexture->GetSizeX(), ExternalTexture->GetSizeY()));
 									PPMInputTextures[InputIndex] = FScreenPassTextureSlice::CreateFromScreenPassTexture(GraphBuilder, ExternalTextureInput);
@@ -257,6 +267,13 @@ void FPPMChainGraphSceneViewExtension::GatherChainGraphProxies
 
 void FPPMChainGraphSceneViewExtension::SubscribeToPostProcessingPass(EPostProcessingPass PassId, FAfterPassCallbackDelegateArray& InOutPassCallbacks, bool bIsPassEnabled)
 {
+	{
+		FScopeLock ScopeLock(&WorldSubsystem->ActiveAccessCriticalSection);
+		if (!WorldSubsystem->ActivePasses.Contains((uint32)PassId + 1))
+		{
+			return;
+		}
+	}
 	InOutPassCallbacks.Add(FAfterPassCallbackDelegate::CreateLambda([this, InPassId = PassId](FRDGBuilder& GraphBuilder, const FSceneView& View, const FPostProcessMaterialInputs& InInputs)
 		{
 			FPostProcessMaterialInputs InOutInputs = InInputs;
diff --git a/Engine/Plugins/Experimental/PostProcessMaterialChainGraph/Source/PostProcessMaterialChainGraph/Private/PPMChainGraphWorldSubsystem.cpp b/Engine/Plugins/Experimental/PostProcessMaterialChainGraph/Source/PostProcessMaterialChainGraph/Private/PPMChainGraphWorldSubsystem.cpp
--- a/Engine/Plugins/Experimental/PostProcessMaterialChainGraph/Source/PostProcessMaterialChainGraph/Private/PPMChainGraphWorldSubsystem.cpp
+++ b/Engine/Plugins/Experimental/PostProcessMaterialChainGraph/Source/PostProcessMaterialChainGraph/Private/PPMChainGraphWorldSubsystem.cpp
@@ -66,6 +66,7 @@ void UPPMChainGraphWorldSubsystem::Deinitialize()
 void UPPMChainGraphWorldSubsystem::Tick(float DeltaTime)
 {
 	Super::Tick(DeltaTime);
+	GatherActivePasses();
 }
 
 void UPPMChainGraphWorldSubsystem::AddPPMChainGraphComponent(TWeakObjectPtr<UPPMChainGraphExecutorComponent> InComponent)
@@ -86,4 +87,27 @@ void UPPMChainGraphWorldSubsystem::RemovePPMChainGraphComponent(TWeakObjectPtr<U
 	}
 }
 
+void UPPMChainGraphWorldSubsystem::GatherActivePasses()
+{
+	TSet<uint32> TempActivePasses;
+	for (const TWeakObjectPtr<UPPMChainGraphExecutorComponent>& PPMChainGraphComponent : PPMChainGraphComponents)
+	{
+		if (!PPMChainGraphComponent.IsValid())
+		{
+			continue;
+		}
+		for (uint32 PassId = 1; PassId <= (uint32)EPPMChainGraphExecutionLocation::AfterVisualizeDepthOfField; PassId++)
+		{
+			if (PPMChainGraphComponent->IsActiveDuringPass_GameThread((EPPMChainGraphExecutionLocation)(PassId)))
+			{
+				TempActivePasses.FindOrAdd(PassId);
+			}
+		}
+	}
+	{
+		FScopeLock ScopeLock(&ActiveAccessCriticalSection);
+		ActivePasses = MoveTemp(TempActivePasses);
+	}
+}
+
 #undef LOCTEXT_NAMESPACE
\ No newline at end of file
diff --git a/Engine/Plugins/Experimental/PostProcessMaterialChainGraph/Source/PostProcessMaterialChainGraph/Private/PPMChainGraphWorldSubsystem.h b/Engine/Plugins/Experimental/PostProcessMaterialChainGraph/Source/PostProcessMaterialChainGraph/Private/PPMChainGraphWorldSubsystem.h
--- a/Engine/Plugins/Experimental/PostProcessMaterialChainGraph/Source/PostProcessMaterialChainGraph/Private/PPMChainGraphWorldSubsystem.h
+++ b/Engine/Plugins/Experimental/PostProcessMaterialChainGraph/Source/PostProcessMaterialChainGraph/Private/PPMChainGraphWorldSubsystem.h
@@ -40,6 +40,9 @@ public:
 
 	/** Removes a PPM Chain Graph component. Typically this means that the component is removed from the world. */
 	void RemovePPMChainGraphComponent(TWeakObjectPtr<UPPMChainGraphExecutorComponent> InComponent);
+private:
+	/** Populates Active passes set to be used by SVE to subscribe to the active passes. */
+	void GatherActivePasses();
 
 private:
 	TSharedPtr<class FPPMChainGraphSceneViewExtension, ESPMode::ThreadSafe> SceneViewExtension;
@@ -50,6 +53,14 @@ private:
 	* All active chain graph components currently present in associated world.
 	*/
 	TSet<TWeakObjectPtr<UPPMChainGraphExecutorComponent>> PPMChainGraphComponents;
+
+	/**
+	* Aggregation of active passes, so that Scene View Extension knows to which passes to subscribe. uint32 value Reflects EPPMChainGraphExecutionLocation
+	*/
+	TSet<uint32> ActivePasses;
+
+	FCriticalSection ActiveAccessCriticalSection;
+
 public:
 	friend class FPPMChainGraphSceneViewExtension;
 };
\ No newline at end of file
diff --git a/Engine/Plugins/Experimental/PostProcessMaterialChainGraph/Source/PostProcessMaterialChainGraph/Public/PPMChainGraph.h b/Engine/Plugins/Experimental/PostProcessMaterialChainGraph/Source/PostProcessMaterialChainGraph/Public/PPMChainGraph.h
--- a/Engine/Plugins/Experimental/PostProcessMaterialChainGraph/Source/PostProcessMaterialChainGraph/Public/PPMChainGraph.h
+++ b/Engine/Plugins/Experimental/PostProcessMaterialChainGraph/Source/PostProcessMaterialChainGraph/Public/PPMChainGraph.h
@@ -52,7 +52,7 @@ enum class EPPMChainGraphExecutionLocation : uint8
 	//AfterSSRInput				UMETA(DisplayName = "After SSR Input"),
 	// 
 	// This is eventually mapped to equivalent execution location enum: ISceneViewExtension::EPostProcessingPass.
-	AfterMotionBlur				UMETA(DisplayName = "After Motion Blur"),
+	AfterMotionBlur = 2			UMETA(DisplayName = "After Motion Blur"),
 	AfterToneMap				UMETA(DisplayName = "After Tonemap"),
 	AfterFXAA					UMETA(DisplayName = "After FXAA"),
 	AfterVisualizeDepthOfField	UMETA(DisplayName = "After Visualize Depth of Field"),
diff --git a/Engine/Plugins/Experimental/PostProcessMaterialChainGraph/Source/PostProcessMaterialChainGraph/Public/PPMChainGraphActor.h b/Engine/Plugins/Experimental/PostProcessMaterialChainGraph/Source/PostProcessMaterialChainGraph/Public/PPMChainGraphActor.h
--- a/Engine/Plugins/Experimental/PostProcessMaterialChainGraph/Source/PostProcessMaterialChainGraph/Public/PPMChainGraphActor.h
+++ b/Engine/Plugins/Experimental/PostProcessMaterialChainGraph/Source/PostProcessMaterialChainGraph/Public/PPMChainGraphActor.h
@@ -15,4 +15,13 @@ class PPMCHAINGRAPH_API APPMChainGraphActor : public AActor
 public:
 	UPROPERTY(Category = MediaPlate, VisibleAnywhere, BlueprintReadOnly)
 	TObjectPtr<UPPMChainGraphExecutorComponent> PPMChainGraphExecutorComponent;
+
+#if WITH_EDITORONLY_DATA
+
+	/** Billboard component for this actor. */
+	UPROPERTY(Transient)
+	TObjectPtr<class UBillboardComponent> SpriteComponent;
+
+#endif // WITH_EDITORONLY_DATA
+
 };
\ No newline at end of file
diff --git a/Engine/Plugins/Experimental/PostProcessMaterialChainGraph/Source/PostProcessMaterialChainGraph/Public/PPMChainGraphComponent.h b/Engine/Plugins/Experimental/PostProcessMaterialChainGraph/Source/PostProcessMaterialChainGraph/Public/PPMChainGraphComponent.h
--- a/Engine/Plugins/Experimental/PostProcessMaterialChainGraph/Source/PostProcessMaterialChainGraph/Public/PPMChainGraphComponent.h
+++ b/Engine/Plugins/Experimental/PostProcessMaterialChainGraph/Source/PostProcessMaterialChainGraph/Public/PPMChainGraphComponent.h
@@ -41,6 +41,9 @@ public:
 	/** Get a copy of Render Proxies for rendering graphs. */
 	TArray<TSharedPtr<FPPMChainGraphProxy>> GetChainGraphRenderProxies(EPPMChainGraphExecutionLocation InPointOfExecution);
 
+	/** Identifies if this component needs to render during the specified pass. */
+	bool IsActiveDuringPass_GameThread(EPPMChainGraphExecutionLocation InPointOfExecution);
+
 public:
 	/** PPM Chain Graph Component can either exclude selected cameras from being rendered or render only in selected camera views. */
 	UPROPERTY(BlueprintReadWrite, EditAnywhere, Category = "Camera View Settings")
