Post Process Material Chain Graph: Added customization that simiplifies input selection for the user and reduces the possibility of invalid input selection.

diff --git a/Engine/Plugins/Experimental/PostProcessMaterialChainGraph/Source/PostProcessMaterialChainGraph/Private/PPMChainGraphComponent.cpp b/Engine/Plugins/Experimental/PostProcessMaterialChainGraph/Source/PostProcessMaterialChainGraph/Private/PPMChainGraphComponent.cpp
--- a/Engine/Plugins/Experimental/PostProcessMaterialChainGraph/Source/PostProcessMaterialChainGraph/Private/PPMChainGraphComponent.cpp
+++ b/Engine/Plugins/Experimental/PostProcessMaterialChainGraph/Source/PostProcessMaterialChainGraph/Private/PPMChainGraphComponent.cpp
@@ -189,9 +189,9 @@ void UPPMChainGraphExecutorComponent::TransferState()
 			bool bInputReferencesOutput = false;
 
 			// Validate inputs.
-			for (TTuple<EPPMChainGraphPPMInputId, FString> KeyValue : Pass.Inputs)
+			for (TTuple<EPPMChainGraphPPMInputId, FPPMChainGraphInput> KeyValue : Pass.Inputs)
 			{
-				if (KeyValue.Value == Pass.TemporaryRenderTargetId)
+				if (KeyValue.Value.InputId == Pass.TemporaryRenderTargetId)
 				{
 					bInputReferencesOutput = true;
 
diff --git a/Engine/Plugins/Experimental/PostProcessMaterialChainGraph/Source/PostProcessMaterialChainGraph/Private/PPMChainGraphSceneViewExtension.cpp b/Engine/Plugins/Experimental/PostProcessMaterialChainGraph/Source/PostProcessMaterialChainGraph/Private/PPMChainGraphSceneViewExtension.cpp
--- a/Engine/Plugins/Experimental/PostProcessMaterialChainGraph/Source/PostProcessMaterialChainGraph/Private/PPMChainGraphSceneViewExtension.cpp
+++ b/Engine/Plugins/Experimental/PostProcessMaterialChainGraph/Source/PostProcessMaterialChainGraph/Private/PPMChainGraphSceneViewExtension.cpp
@@ -109,7 +109,7 @@ namespace
 					EPPMChainGraphPPMInputId InputId = (EPPMChainGraphPPMInputId)(InputIndex + 1);
 					if (Pass->Inputs.Contains(InputId))
 					{
-						FString InputMappedToId = Pass->Inputs[InputId];
+						FString InputMappedToId = Pass->Inputs[InputId].InputId;
 						// Need to exclude current pass.
 						if (PassOutputs.Contains(InputMappedToId))
 						{
diff --git a/Engine/Plugins/Experimental/PostProcessMaterialChainGraph/Source/PostProcessMaterialChainGraph/Public/PPMChainGraph.h b/Engine/Plugins/Experimental/PostProcessMaterialChainGraph/Source/PostProcessMaterialChainGraph/Public/PPMChainGraph.h
--- a/Engine/Plugins/Experimental/PostProcessMaterialChainGraph/Source/PostProcessMaterialChainGraph/Public/PPMChainGraph.h
+++ b/Engine/Plugins/Experimental/PostProcessMaterialChainGraph/Source/PostProcessMaterialChainGraph/Public/PPMChainGraph.h
@@ -59,6 +59,19 @@ enum class EPPMChainGraphExecutionLocation : uint8
 };
 
 
+/**
+* This struct is used for customizing Input and External Texture input selection.
+*/
+USTRUCT(BlueprintType)
+struct FPPMChainGraphInput
+{
+	GENERATED_USTRUCT_BODY()
+
+	UPROPERTY(EditAnywhere, BlueprintReadWrite, Category = Input)
+	FString InputId;
+};
+
+
 USTRUCT(Blueprintable)
 struct FPPMChainGraphPostProcessPass
 {
@@ -70,8 +83,14 @@ struct FPPMChainGraphPostProcessPass
 
 	/** Inputs from previous passes. Map this to Scene Texture node in Post Process Material. */
 	UPROPERTY(BlueprintReadWrite, EditAnywhere, Category = "Post Process Pass Inputs")
-	TMap<EPPMChainGraphPPMInputId, FString> Inputs;
-	
+	TMap<EPPMChainGraphPPMInputId, FPPMChainGraphInput> Inputs;
+
+	/**
+	* Which material should be executed during this pass.
+	*/
+	UPROPERTY(BlueprintReadWrite, EditAnywhere, Category = "Post Process Pass Material")
+	TObjectPtr<UMaterial> PostProcessMaterial;
+
 	/** 
 	* Where should this pass write to. By selecting Temporary Render Target as an option, 
 	* Users can avoid writing directly into Scene Color. For example this pass can operate
@@ -84,12 +103,6 @@ struct FPPMChainGraphPostProcessPass
 	/** Use this to identify the Output of the current pass to be used later. */
 	UPROPERTY(BlueprintReadWrite, EditAnywhere, Category = "Post Process Pass Output", meta = (EditCondition = "Output == EPPMChainGraphOutput::PPMOutput_RenderTarget", EditConditionHides))
 	FString TemporaryRenderTargetId;
-
-	/**
-	* Which material should be executed during this pass.
-	*/
-	UPROPERTY(BlueprintReadWrite, EditAnywhere, Category = "Post Process Pass Material")
-	TObjectPtr<UMaterial> PostProcessMaterial;
 };
 
 
@@ -125,6 +138,6 @@ public:
 	* Post Process Material Passes. Each pass can write into the Scene Color or into a temporary render target that can
 	* be referenced in subsequent passes. At the end of all passes the result is always written into Scene Color.
 	*/
-	UPROPERTY(BlueprintReadWrite, EditAnywhere, Category = "Post Process Passes")
+	UPROPERTY(BlueprintReadWrite, EditAnywhere, Category = "Post Process Passes", meta = (TitleProperty = "TemporaryRenderTargetId"))
 	TArray<FPPMChainGraphPostProcessPass> Passes;
 };
\ No newline at end of file
diff --git a/Engine/Plugins/Experimental/PostProcessMaterialChainGraph/Source/PostProcessMaterialChainGraphEditor/PPMChainGraphEditor.Build.cs b/Engine/Plugins/Experimental/PostProcessMaterialChainGraph/Source/PostProcessMaterialChainGraphEditor/PPMChainGraphEditor.Build.cs
--- a/Engine/Plugins/Experimental/PostProcessMaterialChainGraph/Source/PostProcessMaterialChainGraphEditor/PPMChainGraphEditor.Build.cs
+++ b/Engine/Plugins/Experimental/PostProcessMaterialChainGraph/Source/PostProcessMaterialChainGraphEditor/PPMChainGraphEditor.Build.cs
@@ -14,8 +14,10 @@ namespace UnrealBuildTool.Rules
 					"EditorFramework",
 					"PPMChainGraph",
 					"UnrealEd",
+					"Slate",
+					"SlateCore",
 				});
-			
+
 			PrivateIncludePathModuleNames.AddRange(
 				new string[] {
 					
diff --git a/Engine/Plugins/Experimental/PostProcessMaterialChainGraph/Source/PostProcessMaterialChainGraphEditor/Private/Customizations/PPMChainGraphCustomization.cpp b/Engine/Plugins/Experimental/PostProcessMaterialChainGraph/Source/PostProcessMaterialChainGraphEditor/Private/Customizations/PPMChainGraphCustomization.cpp
new file mode 100644
--- /dev/null
+++ b/Engine/Plugins/Experimental/PostProcessMaterialChainGraph/Source/PostProcessMaterialChainGraphEditor/Private/Customizations/PPMChainGraphCustomization.cpp
@@ -0,0 +1,29 @@
+// Copyright Epic Games, Inc. All Rights Reserved.
+
+#include "PPMChainGraphCustomization.h"
+
+#include "PPMChainGraph.h"
+#include "PPMChainGraphInputCustomization.h"
+#include "Containers/StringConv.h"
+#include "DetailWidgetRow.h"
+#include "DetailLayoutBuilder.h"
+#include "DetailCategoryBuilder.h"
+#include "PropertyHandle.h"
+#include "Widgets/Input/SButton.h"
+#include "Widgets/Input/SComboButton.h"
+
+#define LOCTEXT_NAMESPACE "PPMChainGraphCustomization"
+
+void FPPMChainGraphCustomization::CustomizeDetails(IDetailLayoutBuilder& DetailBuilder)
+{
+	TArray<TWeakObjectPtr<UObject>> Objects;
+	DetailBuilder.GetObjectsBeingCustomized(Objects);
+
+	TWeakObjectPtr<UPPMChainGraph> ChainGraph = Cast<UPPMChainGraph>(Objects[0].Get());
+
+	DetailBuilder.RegisterInstancedCustomPropertyTypeLayout(
+		FPPMChainGraphInput::StaticStruct()->GetFName(),
+		FOnGetPropertyTypeCustomizationInstance::CreateLambda([ChainGraph] { return FPPMChainGraphInputCustomization::MakeInstance(ChainGraph); }));
+}
+
+#undef LOCTEXT_NAMESPACE
diff --git a/Engine/Plugins/Experimental/PostProcessMaterialChainGraph/Source/PostProcessMaterialChainGraphEditor/Private/Customizations/PPMChainGraphCustomization.h b/Engine/Plugins/Experimental/PostProcessMaterialChainGraph/Source/PostProcessMaterialChainGraphEditor/Private/Customizations/PPMChainGraphCustomization.h
new file mode 100644
--- /dev/null
+++ b/Engine/Plugins/Experimental/PostProcessMaterialChainGraph/Source/PostProcessMaterialChainGraphEditor/Private/Customizations/PPMChainGraphCustomization.h
@@ -0,0 +1,19 @@
+// Copyright Epic Games, Inc. All Rights Reserved.
+
+#pragma once
+
+#include "CoreMinimal.h"
+#include "IDetailCustomization.h"
+
+
+class FPPMChainGraphCustomization : public IDetailCustomization
+{
+public:
+	static TSharedRef<IDetailCustomization> MakeInstance()
+	{
+		return MakeShared<FPPMChainGraphCustomization>();
+	}
+
+	virtual void CustomizeDetails(IDetailLayoutBuilder& DetailBuilder) override;
+};
+
diff --git a/Engine/Plugins/Experimental/PostProcessMaterialChainGraph/Source/PostProcessMaterialChainGraphEditor/Private/Customizations/PPMChainGraphInputCustomization.cpp b/Engine/Plugins/Experimental/PostProcessMaterialChainGraph/Source/PostProcessMaterialChainGraphEditor/Private/Customizations/PPMChainGraphInputCustomization.cpp
new file mode 100644
--- /dev/null
+++ b/Engine/Plugins/Experimental/PostProcessMaterialChainGraph/Source/PostProcessMaterialChainGraphEditor/Private/Customizations/PPMChainGraphInputCustomization.cpp
@@ -0,0 +1,171 @@
+// Copyright Epic Games, Inc. All Rights Reserved.
+
+#include "PPMChainGraphInputCustomization.h"
+
+#include "Containers/StringConv.h"
+#include "DetailWidgetRow.h"
+#include "IPropertyUtilities.h"
+#include "PropertyHandle.h"
+#include "Widgets/Input/SComboButton.h"
+
+#define LOCTEXT_NAMESPACE "PPMChainGraphInputCustomization"
+
+FPPMChainGraphInputCustomization::FPPMChainGraphInputCustomization(TWeakObjectPtr<UPPMChainGraph> InPPMChainGraph)
+	: PPMChainGraphPtr(MoveTemp(InPPMChainGraph))
+{
+}
+
+void FPPMChainGraphInputCustomization::CustomizeHeader(TSharedRef<IPropertyHandle> InPropertyHandle, FDetailWidgetRow& HeaderRow, IPropertyTypeCustomizationUtils& CustomizationUtils)
+{
+	CachedInputsProperty = InPropertyHandle;
+
+	TSharedPtr<IPropertyHandle> InputsArrayProperty = InPropertyHandle->GetParentHandle();
+
+	CachedPassProperty = InputsArrayProperty->GetParentHandle();
+
+	if (InPropertyHandle->GetNumPerObjectValues() == 1 && InPropertyHandle->IsValidHandle())
+	{
+		FProperty* Property = InPropertyHandle->GetProperty();
+		check(Property && CastField<FStructProperty>(Property) && CastField<FStructProperty>(Property)->Struct && CastField<FStructProperty>(Property)->Struct->IsChildOf(FPPMChainGraphInput::StaticStruct()));
+
+		TArray<void*> RawData;
+		InPropertyHandle->AccessRawData(RawData);
+
+		check(RawData.Num() == 1);
+		FPPMChainGraphInput* InputValue = reinterpret_cast<FPPMChainGraphInput*>(RawData[0]);
+
+		check(InputValue);
+		TSharedPtr<IPropertyUtilities> PropertyUtils = CustomizationUtils.GetPropertyUtilities();
+
+		HeaderRow
+			.NameContent()
+			[
+				InPropertyHandle->CreatePropertyNameWidget()
+			]
+			.ValueContent()
+			.MaxDesiredWidth(512)
+			[
+				SNew(SHorizontalBox)
+				+ SHorizontalBox::Slot()
+				.FillWidth(1.0f)
+				.VAlign(VAlign_Center)
+				[
+					SNew(STextBlock)
+					.Text(MakeAttributeLambda([InputValue] {
+						const FString InputId = InputValue->InputId;
+							
+						if(!InputId.IsEmpty())
+						{
+							return FText::FromString(InputId);
+						}
+
+						return LOCTEXT("NotSelected", "<Not Selected>");
+					}))
+				]
+				+ SHorizontalBox::Slot()
+				.AutoWidth()
+				.Padding(FMargin(4.0f, 0.0f, 0.0f, 0.0f))
+				.VAlign(VAlign_Center)
+				[
+					SNew(SComboButton)
+					.OnGetMenuContent(this, &FPPMChainGraphInputCustomization::PopulateComboBox)
+					.ContentPadding(FMargin(4.0, 2.0))
+				]
+			].IsEnabled(MakeAttributeLambda([=] { return !InPropertyHandle->IsEditConst() && PropertyUtils->IsPropertyEditingEnabled(); }));
+	}
+}
+
+void FPPMChainGraphInputCustomization::AddComboBoxEntry(FMenuBuilder& InMenuBuilder, const FPPMChainGraphInput& InTextureInput)
+{
+	InMenuBuilder.AddMenuEntry(
+		FText::FromString(InTextureInput.InputId),
+		FText::FromString(InTextureInput.InputId),
+		FSlateIcon(),
+		FUIAction(
+			FExecuteAction::CreateLambda([this, InTextureInput]
+			{
+				if (FStructProperty* StructProperty = CastField<FStructProperty>(CachedInputsProperty->GetProperty()))
+				{
+					check(true);
+					TArray<void*> RawData;
+					CachedInputsProperty->AccessRawData(RawData);
+					FPPMChainGraphInput* TextureInput = reinterpret_cast<FPPMChainGraphInput*>(RawData[0]);
+
+					FString TextValue;
+					StructProperty->Struct->ExportText(TextValue, &InTextureInput, TextureInput, nullptr, EPropertyPortFlags::PPF_None, nullptr);
+					ensure(CachedInputsProperty->SetValueFromFormattedString(TextValue, EPropertyValueSetFlags::DefaultFlags) == FPropertyAccess::Result::Success);
+				}
+			}
+			),
+			FCanExecuteAction(),
+			FIsActionChecked::CreateLambda([this, InTextureInput]
+			{
+				return true;
+			})
+		),
+		NAME_None,
+		EUserInterfaceActionType::RadioButton
+		);
+}
+
+TSharedRef<SWidget> FPPMChainGraphInputCustomization::PopulateComboBox()
+{
+	FMenuBuilder MenuBuilder(true, nullptr);
+	TArray<FString> ExistingSubMenus;
+
+	MenuBuilder.BeginSection("AllPossibleInputs", LOCTEXT("InputsSelection", "Texture Inputs"));
+	{
+		TArray<void*> RawData;
+		CachedPassProperty->AccessRawData(RawData);
+		FPPMChainGraphPostProcessPass* CurrentPass = reinterpret_cast<FPPMChainGraphPostProcessPass*>(RawData[0]);
+
+		if (PPMChainGraphPtr.IsValid())
+		{
+			TArray<FString> InputIds;
+			for (TTuple<FString, TObjectPtr<UTexture2D>> ExternalTexture : PPMChainGraphPtr->ExternalTextures)
+			{
+				InputIds.Add(ExternalTexture.Key);
+				const FPPMChainGraphInput GraphInputOption = { ExternalTexture.Key };
+				
+				AddComboBoxEntry(MenuBuilder, GraphInputOption);
+			}
+
+			for (const FPPMChainGraphPostProcessPass& Pass : PPMChainGraphPtr->Passes)
+			{
+				if (Pass.TemporaryRenderTargetId.IsEmpty())
+				{
+					continue;
+				}
+
+				if (InputIds.Contains(Pass.TemporaryRenderTargetId))
+				{
+					continue;
+				}
+
+				// We collect passes up to the current one and ignore anything that comes after.
+				if (CurrentPass == &Pass)
+				{
+					break;
+				}
+
+				InputIds.Add(Pass.TemporaryRenderTargetId);
+				const FPPMChainGraphInput GraphInputOption = { Pass.TemporaryRenderTargetId };
+				AddComboBoxEntry(MenuBuilder, GraphInputOption);
+			}
+
+			if (InputIds.Num() <= 0)
+			{
+				MenuBuilder.AddWidget(SNullWidget::NullWidget, LOCTEXT("NoValidInputs", "No valid inputs found"), false, false);
+			}
+		}
+		else
+		{
+			MenuBuilder.AddWidget(SNullWidget::NullWidget, LOCTEXT("InvalidChainGraph", "Invalid PPM Chain Graph object selected"), false, false);
+		}
+	}
+	MenuBuilder.EndSection();
+
+	return MenuBuilder.MakeWidget();
+}
+
+#undef LOCTEXT_NAMESPACE
diff --git a/Engine/Plugins/Experimental/PostProcessMaterialChainGraph/Source/PostProcessMaterialChainGraphEditor/Private/Customizations/PPMChainGraphInputCustomization.h b/Engine/Plugins/Experimental/PostProcessMaterialChainGraph/Source/PostProcessMaterialChainGraphEditor/Private/Customizations/PPMChainGraphInputCustomization.h
new file mode 100644
--- /dev/null
+++ b/Engine/Plugins/Experimental/PostProcessMaterialChainGraph/Source/PostProcessMaterialChainGraphEditor/Private/Customizations/PPMChainGraphInputCustomization.h
@@ -0,0 +1,42 @@
+// Copyright Epic Games, Inc. All Rights Reserved.
+
+#pragma once
+
+#include "CoreMinimal.h"
+#include "Engine/EngineTypes.h"
+#include "Framework/MultiBox/MultiBoxBuilder.h"
+#include "IPropertyTypeCustomization.h"
+#include "PPMChainGraph.h"
+#include "Widgets/SWidget.h"
+
+class UPPMChainGraph;
+
+class FPPMChainGraphInputCustomization : public IPropertyTypeCustomization
+{
+public:
+	FPPMChainGraphInputCustomization(TWeakObjectPtr<UPPMChainGraph> InConfigurationObjectProperty);
+
+	/** IPropertyTypeCustomization interface stub. */
+	virtual void CustomizeChildren(TSharedRef<class IPropertyHandle> InPropertyHandle, class IDetailChildrenBuilder& StructBuilder, IPropertyTypeCustomizationUtils& PropertyTypeCustomizationUtils) override {}
+
+	static TSharedRef<IPropertyTypeCustomization> MakeInstance(TWeakObjectPtr<UPPMChainGraph> InConfigurationObjectProperty)
+	{
+		return MakeShared<FPPMChainGraphInputCustomization>(MoveTemp(InConfigurationObjectProperty));
+	}
+
+	virtual void CustomizeHeader(TSharedRef<class IPropertyHandle> InPropertyHandle, class FDetailWidgetRow& HeaderRow, IPropertyTypeCustomizationUtils& PropertyTypeCustomizationUtils) override;
+
+protected:
+
+	TSharedPtr<IPropertyHandle> CachedInputsProperty;
+	TSharedPtr<IPropertyHandle> CachedPassProperty;
+
+	TWeakObjectPtr<UPPMChainGraph> PPMChainGraphPtr;
+
+private:
+	void AddComboBoxEntry(FMenuBuilder& InMenuBuilder, const FPPMChainGraphInput& InTextureInput);
+
+	TSharedRef<SWidget> PopulateComboBox();
+};
+
+
diff --git a/Engine/Plugins/Experimental/PostProcessMaterialChainGraph/Source/PostProcessMaterialChainGraphEditor/Private/PPMChainGraphEditorModule.cpp b/Engine/Plugins/Experimental/PostProcessMaterialChainGraph/Source/PostProcessMaterialChainGraphEditor/Private/PPMChainGraphEditorModule.cpp
--- a/Engine/Plugins/Experimental/PostProcessMaterialChainGraph/Source/PostProcessMaterialChainGraphEditor/Private/PPMChainGraphEditorModule.cpp
+++ b/Engine/Plugins/Experimental/PostProcessMaterialChainGraph/Source/PostProcessMaterialChainGraphEditor/Private/PPMChainGraphEditorModule.cpp
@@ -2,10 +2,14 @@
 
 #include "PPMChainGraphEditorModule.h"
 #include "Modules/ModuleManager.h"
-
+#include "Customizations/PPMChainGraphCustomization.h"
+#include "PropertyEditorModule.h"
+#include "PPMChainGraph.h"
 
 void FPPMChainGraphEditorModule::StartupModule()
 {
+	FPropertyEditorModule& PropertyModule = FModuleManager::LoadModuleChecked<FPropertyEditorModule>("PropertyEditor");
+	PropertyModule.RegisterCustomClassLayout(UPPMChainGraph::StaticClass()->GetFName(), FOnGetDetailCustomizationInstance::CreateStatic(&FPPMChainGraphCustomization::MakeInstance));
 }
 
 void FPPMChainGraphEditorModule::ShutdownModule()
