Post Process Material Chain Graphs: Fixed an issue with module packaging.

diff --git a/Engine/Plugins/Experimental/PostProcessMaterialChainGraph/Source/PostProcessMaterialChainGraph/Private/PPMChainGraphModule.cpp b/Engine/Plugins/Experimental/PostProcessMaterialChainGraph/Source/PostProcessMaterialChainGraph/Private/PPMChainGraphModule.cpp
--- a/Engine/Plugins/Experimental/PostProcessMaterialChainGraph/Source/PostProcessMaterialChainGraph/Private/PPMChainGraphModule.cpp
+++ b/Engine/Plugins/Experimental/PostProcessMaterialChainGraph/Source/PostProcessMaterialChainGraph/Private/PPMChainGraphModule.cpp
@@ -10,4 +10,4 @@ void FPostProcessMaterialChainGraphModule::ShutdownModule()
 {
 }
 	
-IMPLEMENT_MODULE(FPostProcessMaterialChainGraphModule, PostProcessMaterialChainGraph)
\ No newline at end of file
+IMPLEMENT_MODULE(FPostProcessMaterialChainGraphModule, PPMChainGraph)
\ No newline at end of file
