PPMChainGraph: Remove unused private renderer include.

#rnx

diff --git a/Engine/Plugins/Experimental/PostProcessMaterialChainGraph/Source/PostProcessMaterialChainGraph/Private/PPMChainGraphSceneViewExtension.cpp b/Engine/Plugins/Experimental/PostProcessMaterialChainGraph/Source/PostProcessMaterialChainGraph/Private/PPMChainGraphSceneViewExtension.cpp
--- a/Engine/Plugins/Experimental/PostProcessMaterialChainGraph/Source/PostProcessMaterialChainGraph/Private/PPMChainGraphSceneViewExtension.cpp
+++ b/Engine/Plugins/Experimental/PostProcessMaterialChainGraph/Source/PostProcessMaterialChainGraph/Private/PPMChainGraphSceneViewExtension.cpp
@@ -15,7 +15,6 @@
 #include "PPMChainGraphWorldSubsystem.h"
 #include "PPMChainGraph.h"
 #include "RHI.h"
-#include "Runtime/Renderer/Private/SceneTextureParameters.h"
 #include "ScreenPass.h"
 #include "Engine/Texture2D.h"
 #include "TextureResource.h"
