Avoid additional copy to back buffer in post process material chain view extension

diff --git a/Engine/Plugins/Compositing/OpenColorIO/Source/OpenColorIO/Private/OpenColorIODisplayExtension.cpp b/Engine/Plugins/Compositing/OpenColorIO/Source/OpenColorIO/Private/OpenColorIODisplayExtension.cpp
--- a/Engine/Plugins/Compositing/OpenColorIO/Source/OpenColorIO/Private/OpenColorIODisplayExtension.cpp
+++ b/Engine/Plugins/Compositing/OpenColorIO/Source/OpenColorIO/Private/OpenColorIODisplayExtension.cpp
@@ -66,7 +66,7 @@ void FOpenColorIODisplayExtension::SetupView(FSceneViewFamily& InViewFamily, FSc
 	);
 }
 
-void FOpenColorIODisplayExtension::SubscribeToPostProcessingPass(EPostProcessingPass PassId, FAfterPassCallbackDelegateArray& InOutPassCallbacks, bool bIsPassEnabled)
+void FOpenColorIODisplayExtension::SubscribeToPostProcessingPass(EPostProcessingPass PassId, const FSceneView& View, FAfterPassCallbackDelegateArray& InOutPassCallbacks, bool bIsPassEnabled)
 {
 	if (PassId == EPostProcessingPass::Tonemap)
 	{
diff --git a/Engine/Plugins/Compositing/OpenColorIO/Source/OpenColorIO/Public/OpenColorIODisplayExtension.h b/Engine/Plugins/Compositing/OpenColorIO/Source/OpenColorIO/Public/OpenColorIODisplayExtension.h
--- a/Engine/Plugins/Compositing/OpenColorIO/Source/OpenColorIO/Public/OpenColorIODisplayExtension.h
+++ b/Engine/Plugins/Compositing/OpenColorIO/Source/OpenColorIO/Public/OpenColorIODisplayExtension.h
@@ -33,7 +33,7 @@ public:
 	virtual void PreRenderViewFamily_RenderThread(FRDGBuilder& GraphBuilder, FSceneViewFamily& InViewFamily) {};
 	virtual void PreRenderView_RenderThread(FRDGBuilder& GraphBuilder, FSceneView& InView) {};
 
-	virtual void SubscribeToPostProcessingPass(EPostProcessingPass PassId, FAfterPassCallbackDelegateArray& InOutPassCallbacks, bool bIsPassEnabled) override;
+	virtual void SubscribeToPostProcessingPass(EPostProcessingPass PassId, const FSceneView& View, FAfterPassCallbackDelegateArray& InOutPassCallbacks, bool bIsPassEnabled) override;
 	virtual bool IsActiveThisFrame_Internal(const FSceneViewExtensionContext& Context) const override;
 	//~End ISceneVIewExtension interface
 
diff --git a/Engine/Plugins/Experimental/CineCameraSceneCapture/Source/CineCameraSceneCapture/Private/CineCameraSceneCaptureComponent.cpp b/Engine/Plugins/Experimental/CineCameraSceneCapture/Source/CineCameraSceneCapture/Private/CineCameraSceneCaptureComponent.cpp
--- a/Engine/Plugins/Experimental/CineCameraSceneCapture/Source/CineCameraSceneCapture/Private/CineCameraSceneCaptureComponent.cpp
+++ b/Engine/Plugins/Experimental/CineCameraSceneCapture/Source/CineCameraSceneCapture/Private/CineCameraSceneCaptureComponent.cpp
@@ -93,7 +93,7 @@ public:
 		DeltaTime = InDeltaTime;
 	}
 
-	virtual void SubscribeToPostProcessingPass(EPostProcessingPass PassId, FAfterPassCallbackDelegateArray& InOutPassCallbacks, bool bIsPassEnabled)
+	virtual void SubscribeToPostProcessingPass(EPostProcessingPass PassId, const FSceneView& View, FAfterPassCallbackDelegateArray& InOutPassCallbacks, bool bIsPassEnabled)
 	{
 		if (!bIsOcioEnabledRenderThread)
 		{
diff --git a/Engine/Plugins/Experimental/Compositing/HoldoutComposite/Source/HoldoutComposite/Private/HoldoutCompositeSceneViewExtension.cpp b/Engine/Plugins/Experimental/Compositing/HoldoutComposite/Source/HoldoutComposite/Private/HoldoutCompositeSceneViewExtension.cpp
--- a/Engine/Plugins/Experimental/Compositing/HoldoutComposite/Source/HoldoutComposite/Private/HoldoutCompositeSceneViewExtension.cpp
+++ b/Engine/Plugins/Experimental/Compositing/HoldoutComposite/Source/HoldoutComposite/Private/HoldoutCompositeSceneViewExtension.cpp
@@ -310,7 +310,7 @@ void FHoldoutCompositeSceneViewExtension::SetupView(FSceneViewFamily& InViewFami
 	WorldPtr.Get()->Scene->AddCustomRenderPass(&InViewFamily, PassInput);
 }
 
-void FHoldoutCompositeSceneViewExtension::SubscribeToPostProcessingPass(EPostProcessingPass PassId, FAfterPassCallbackDelegateArray& InOutPassCallbacks, bool bIsPassEnabled)
+void FHoldoutCompositeSceneViewExtension::SubscribeToPostProcessingPass(EPostProcessingPass PassId, const FSceneView& View, FAfterPassCallbackDelegateArray& InOutPassCallbacks, bool bIsPassEnabled)
 {
 	if (PassId == EPostProcessingPass::Tonemap)
 	{
diff --git a/Engine/Plugins/Experimental/Compositing/HoldoutComposite/Source/HoldoutComposite/Private/HoldoutCompositeSceneViewExtension.h b/Engine/Plugins/Experimental/Compositing/HoldoutComposite/Source/HoldoutComposite/Private/HoldoutCompositeSceneViewExtension.h
--- a/Engine/Plugins/Experimental/Compositing/HoldoutComposite/Source/HoldoutComposite/Private/HoldoutCompositeSceneViewExtension.h
+++ b/Engine/Plugins/Experimental/Compositing/HoldoutComposite/Source/HoldoutComposite/Private/HoldoutCompositeSceneViewExtension.h
@@ -29,7 +29,7 @@ public:
 	virtual void SetupViewFamily(FSceneViewFamily& InViewFamily) override;
 	virtual void SetupView(FSceneViewFamily& InViewFamily, FSceneView& InView) override;
 	virtual void BeginRenderViewFamily(FSceneViewFamily& InViewFamily) override {};
-	virtual void SubscribeToPostProcessingPass(EPostProcessingPass PassId, FAfterPassCallbackDelegateArray& InOutPassCallbacks, bool bIsPassEnabled) override;
+	virtual void SubscribeToPostProcessingPass(EPostProcessingPass PassId, const FSceneView& View, FAfterPassCallbackDelegateArray& InOutPassCallbacks, bool bIsPassEnabled) override;
 	virtual void PostRenderViewFamily_RenderThread(FRDGBuilder& GraphBuilder, FSceneViewFamily& InViewFamily) override;
 	//~ End ISceneViewExtension Interface
 
diff --git a/Engine/Plugins/Experimental/PostProcessMaterialChainGraph/Source/PostProcessMaterialChainGraph/Private/PPMChainGraphSceneViewExtension.cpp b/Engine/Plugins/Experimental/PostProcessMaterialChainGraph/Source/PostProcessMaterialChainGraph/Private/PPMChainGraphSceneViewExtension.cpp
--- a/Engine/Plugins/Experimental/PostProcessMaterialChainGraph/Source/PostProcessMaterialChainGraph/Private/PPMChainGraphSceneViewExtension.cpp
+++ b/Engine/Plugins/Experimental/PostProcessMaterialChainGraph/Source/PostProcessMaterialChainGraph/Private/PPMChainGraphSceneViewExtension.cpp
@@ -264,7 +264,7 @@ void FPPMChainGraphSceneViewExtension::GatherChainGraphProxies
 	}
 }
 
-void FPPMChainGraphSceneViewExtension::SubscribeToPostProcessingPass(EPostProcessingPass PassId, FAfterPassCallbackDelegateArray& InOutPassCallbacks, bool bIsPassEnabled)
+void FPPMChainGraphSceneViewExtension::SubscribeToPostProcessingPass(EPostProcessingPass PassId, const FSceneView& View, FAfterPassCallbackDelegateArray& InOutPassCallbacks, bool bIsPassEnabled)
 {
 	{
 		FScopeLock ScopeLock(&WorldSubsystem->ActiveAccessCriticalSection);
@@ -273,27 +273,29 @@ void FPPMChainGraphSceneViewExtension::SubscribeToPostProcessingPass(EPostProces
 			return;
 		}
 	}
-	InOutPassCallbacks.Add(FAfterPassCallbackDelegate::CreateLambda([this, InPassId = PassId](FRDGBuilder& GraphBuilder, const FSceneView& View, const FPostProcessMaterialInputs& InInputs)
-		{
-			FPostProcessMaterialInputs InOutInputs = InInputs;
-			return FPPMChainGraphSceneViewExtension::AfterPostProcessPass_RenderThread(GraphBuilder, View, InOutInputs, InPassId);
-		}));
+
+	TArray<TSharedPtr<FPPMChainGraphProxy>> ChainGraphProxies;
+	EPPMChainGraphExecutionLocation PointOfExecution = (EPPMChainGraphExecutionLocation)((int)PassId + 1);
+	GatherChainGraphProxies(ChainGraphProxies, View, *View.Family, PointOfExecution);
+
+	if (!ChainGraphProxies.IsEmpty())
+	{
+		InOutPassCallbacks.Add(FAfterPassCallbackDelegate::CreateLambda([this, InPassId = PassId, ChainGraphProxies = MoveTemp(ChainGraphProxies)](FRDGBuilder& GraphBuilder, const FSceneView& View, const FPostProcessMaterialInputs& InInputs)
+			{
+				FPostProcessMaterialInputs InOutInputs = InInputs;
+				return FPPMChainGraphSceneViewExtension::AfterPostProcessPass_RenderThread(GraphBuilder, View, InOutInputs, InPassId, ChainGraphProxies);
+			}));
+	}
 }
 
-FScreenPassTexture FPPMChainGraphSceneViewExtension::AfterPostProcessPass_RenderThread(FRDGBuilder& GraphBuilder, const FSceneView& View, FPostProcessMaterialInputs& InOutInputs, EPostProcessingPass InCurrentPass)
+FScreenPassTexture FPPMChainGraphSceneViewExtension::AfterPostProcessPass_RenderThread(FRDGBuilder& GraphBuilder, const FSceneView& View, FPostProcessMaterialInputs& InOutInputs, EPostProcessingPass InCurrentPass, const TArray<TSharedPtr<FPPMChainGraphProxy>>& InChainGraphProxy)
 {
 	const FSceneViewFamily& ViewFamily = *View.Family;
 
 	TArray<TSharedPtr<FPPMChainGraphProxy>> ChainGraphProxies;
 
 	EPPMChainGraphExecutionLocation PointOfExecution = (EPPMChainGraphExecutionLocation)((int)InCurrentPass + 1);
-	GatherChainGraphProxies(ChainGraphProxies, View, ViewFamily, PointOfExecution);
 	FScreenPassRenderTarget OverrideOutput = InOutInputs.OverrideOutput;
-	if (ChainGraphProxies.IsEmpty())
-	{
-		// Don't need to modify anything, just return the untouched scene color texture back to post processing.
-		return InOutInputs.ReturnUntouchedSceneColorForPostProcessing(GraphBuilder);
-	}
 
 	const FIntRect PrimaryViewRect = UE::FXRenderingUtils::GetRawViewRectUnsafe(View);
 	const FScreenPassTexture SceneColor = FScreenPassTexture::CopyFromSlice(GraphBuilder, InOutInputs.GetInput(EPostProcessMaterialInput::SceneColor));
@@ -345,6 +347,7 @@ FScreenPassTexture FPPMChainGraphSceneViewExtension::AfterPostProcessPass_Render
 
 void FPPMChainGraphSceneViewExtension::PrePostProcessPass_RenderThread(FRDGBuilder& GraphBuilder, const FSceneView& View, const FPostProcessingInputs& Inputs)
 {
+	CachedView = &View;
 	const FSceneViewFamily& ViewFamily = *View.Family;
 
 	TArray<TSharedPtr<FPPMChainGraphProxy>> ChainGraphProxies;
diff --git a/Engine/Plugins/Experimental/PostProcessMaterialChainGraph/Source/PostProcessMaterialChainGraph/Private/PPMChainGraphSceneViewExtension.h b/Engine/Plugins/Experimental/PostProcessMaterialChainGraph/Source/PostProcessMaterialChainGraph/Private/PPMChainGraphSceneViewExtension.h
--- a/Engine/Plugins/Experimental/PostProcessMaterialChainGraph/Source/PostProcessMaterialChainGraph/Private/PPMChainGraphSceneViewExtension.h
+++ b/Engine/Plugins/Experimental/PostProcessMaterialChainGraph/Source/PostProcessMaterialChainGraph/Private/PPMChainGraphSceneViewExtension.h
@@ -15,17 +15,23 @@ public:
 	FPPMChainGraphSceneViewExtension(const FAutoRegister& AutoRegister, UPPMChainGraphWorldSubsystem* InWorldSubsystem);
 
 	//~ Begin FSceneViewExtensionBase Interface
-	virtual void SubscribeToPostProcessingPass(EPostProcessingPass PassId, FAfterPassCallbackDelegateArray& InOutPassCallbacks, bool bIsPassEnabled) override;
+	virtual void SubscribeToPostProcessingPass(EPostProcessingPass PassId, const FSceneView& View, FAfterPassCallbackDelegateArray& InOutPassCallbacks, bool bIsPassEnabled) override;
 
 	virtual void SetupViewFamily(FSceneViewFamily& InViewFamily) override {};
 	virtual void SetupView(FSceneViewFamily& InViewFamily, FSceneView& InView) override {};
 	virtual void BeginRenderViewFamily(FSceneViewFamily& InViewFamily) override {};
 	virtual void PrePostProcessPass_RenderThread(FRDGBuilder& GraphBuilder, const FSceneView& View, const FPostProcessingInputs& Inputs) override;
-	FScreenPassTexture AfterPostProcessPass_RenderThread(FRDGBuilder& GraphBuilder, const FSceneView& View, FPostProcessMaterialInputs& InOutInputs, EPostProcessingPass InCurrentPass);
+	FScreenPassTexture AfterPostProcessPass_RenderThread(FRDGBuilder& GraphBuilder, const FSceneView& View, FPostProcessMaterialInputs& InOutInputs, EPostProcessingPass InCurrentPass, const TArray<TSharedPtr<FPPMChainGraphProxy>>& OutChainGraphProxies);
+
+	virtual void PostRenderView_RenderThread(FRDGBuilder& GraphBuilder, FSceneView& InView) override
+	{
+		CachedView = nullptr;
+	}
 
 private:
 	void GatherChainGraphProxies(TArray<TSharedPtr<FPPMChainGraphProxy>>& OutChainGraphProxies, const FSceneView& InView, const FSceneViewFamily& InViewFamily, EPPMChainGraphExecutionLocation InPointOfExecution);
 
 private:
+	const FSceneView* CachedView = nullptr;
 	TWeakObjectPtr<UPPMChainGraphWorldSubsystem> WorldSubsystem;
 };
\ No newline at end of file
diff --git a/Engine/Plugins/Media/MediaIOFramework/Source/MediaIOCore/Private/MediaCaptureSceneViewExtension.h b/Engine/Plugins/Media/MediaIOFramework/Source/MediaIOCore/Private/MediaCaptureSceneViewExtension.h
--- a/Engine/Plugins/Media/MediaIOFramework/Source/MediaIOCore/Private/MediaCaptureSceneViewExtension.h
+++ b/Engine/Plugins/Media/MediaIOFramework/Source/MediaIOCore/Private/MediaCaptureSceneViewExtension.h
@@ -73,7 +73,7 @@ public:
 	};
 	virtual void BeginRenderViewFamily(FSceneViewFamily& InViewFamily) override {}; 
 	
-	virtual void SubscribeToPostProcessingPass(EPostProcessingPass PassId, FAfterPassCallbackDelegateArray& InOutPassCallbacks, bool bIsPassEnabled) override
+	virtual void SubscribeToPostProcessingPass(EPostProcessingPass PassId, const FSceneView& View, FAfterPassCallbackDelegateArray& InOutPassCallbacks, bool bIsPassEnabled) override
 	{
 		if ((CapturePhase == EMediaCapturePhase::AfterMotionBlur && PassId == EPostProcessingPass::MotionBlur)
 			|| (CapturePhase == EMediaCapturePhase::AfterToneMap && PassId == EPostProcessingPass::Tonemap)
diff --git a/Engine/Plugins/Runtime/nDisplay/Source/DisplayCluster/Private/Render/Viewport/DisplayClusterViewportFrameStatsViewExtension.cpp b/Engine/Plugins/Runtime/nDisplay/Source/DisplayCluster/Private/Render/Viewport/DisplayClusterViewportFrameStatsViewExtension.cpp
--- a/Engine/Plugins/Runtime/nDisplay/Source/DisplayCluster/Private/Render/Viewport/DisplayClusterViewportFrameStatsViewExtension.cpp
+++ b/Engine/Plugins/Runtime/nDisplay/Source/DisplayCluster/Private/Render/Viewport/DisplayClusterViewportFrameStatsViewExtension.cpp
@@ -81,7 +81,7 @@ FDisplayClusterViewportFrameStatsViewExtension::FDisplayClusterViewportFrameStat
 	, Configuration(InConfiguration)
 { }
 
-void FDisplayClusterViewportFrameStatsViewExtension::SubscribeToPostProcessingPass(EPostProcessingPass PassId, FAfterPassCallbackDelegateArray& InOutPassCallbacks, bool bIsPassEnabled)
+void FDisplayClusterViewportFrameStatsViewExtension::SubscribeToPostProcessingPass(EPostProcessingPass PassId, const FSceneView& View, FAfterPassCallbackDelegateArray& InOutPassCallbacks, bool bIsPassEnabled)
 {
 	if (IsActive() && PassId == EPostProcessingPass::Tonemap)
 	{
diff --git a/Engine/Plugins/Runtime/nDisplay/Source/DisplayCluster/Private/Render/Viewport/DisplayClusterViewportFrameStatsViewExtension.h b/Engine/Plugins/Runtime/nDisplay/Source/DisplayCluster/Private/Render/Viewport/DisplayClusterViewportFrameStatsViewExtension.h
--- a/Engine/Plugins/Runtime/nDisplay/Source/DisplayCluster/Private/Render/Viewport/DisplayClusterViewportFrameStatsViewExtension.h
+++ b/Engine/Plugins/Runtime/nDisplay/Source/DisplayCluster/Private/Render/Viewport/DisplayClusterViewportFrameStatsViewExtension.h
@@ -21,7 +21,7 @@ public:
 	virtual void SetupView(FSceneViewFamily& InViewFamily, FSceneView& InView) override { }
 	virtual void BeginRenderViewFamily(FSceneViewFamily& InViewFamily) override;
 
-	void SubscribeToPostProcessingPass(EPostProcessingPass PassId, FAfterPassCallbackDelegateArray& InOutPassCallbacks, bool bIsPassEnabled);
+	void SubscribeToPostProcessingPass(EPostProcessingPass PassId, const FSceneView& View, FAfterPassCallbackDelegateArray& InOutPassCallbacks, bool bIsPassEnabled);
 	FScreenPassTexture PostProcessPassAfterTonemap_RenderThread(FRDGBuilder& GraphBuilder, const FSceneView& View, const FPostProcessMaterialInputs& InOutInputs);
 
 	virtual int32 GetPriority() const override { return DISPLAYCLUSTER_SCENE_DEBUG_VIEW_EXTENSION_PRIORITY; }
diff --git a/Engine/Plugins/Runtime/nDisplay/Source/DisplayCluster/Private/Render/Viewport/DisplayClusterViewportManagerViewExtension.cpp b/Engine/Plugins/Runtime/nDisplay/Source/DisplayCluster/Private/Render/Viewport/DisplayClusterViewportManagerViewExtension.cpp
--- a/Engine/Plugins/Runtime/nDisplay/Source/DisplayCluster/Private/Render/Viewport/DisplayClusterViewportManagerViewExtension.cpp
+++ b/Engine/Plugins/Runtime/nDisplay/Source/DisplayCluster/Private/Render/Viewport/DisplayClusterViewportManagerViewExtension.cpp
@@ -59,7 +59,7 @@ void FDisplayClusterViewportManagerViewExtension::PreRenderView_RenderThread(FRD
 	}
 }
 
-void FDisplayClusterViewportManagerViewExtension::SubscribeToPostProcessingPass(EPostProcessingPass PassId, FAfterPassCallbackDelegateArray& InOutPassCallbacks, bool bIsPassEnabled)
+void FDisplayClusterViewportManagerViewExtension::SubscribeToPostProcessingPass(EPostProcessingPass PassId, const FSceneView& View, FAfterPassCallbackDelegateArray& InOutPassCallbacks, bool bIsPassEnabled)
 {
 	if (!IsActive())
 	{
diff --git a/Engine/Plugins/Runtime/nDisplay/Source/DisplayCluster/Private/Render/Viewport/DisplayClusterViewportManagerViewExtension.h b/Engine/Plugins/Runtime/nDisplay/Source/DisplayCluster/Private/Render/Viewport/DisplayClusterViewportManagerViewExtension.h
--- a/Engine/Plugins/Runtime/nDisplay/Source/DisplayCluster/Private/Render/Viewport/DisplayClusterViewportManagerViewExtension.h
+++ b/Engine/Plugins/Runtime/nDisplay/Source/DisplayCluster/Private/Render/Viewport/DisplayClusterViewportManagerViewExtension.h
@@ -36,7 +36,7 @@ public:
 
 	virtual void PreRenderView_RenderThread(FRDGBuilder& GraphBuilder, FSceneView& InView) override;
 
-	virtual void SubscribeToPostProcessingPass(EPostProcessingPass PassId, FAfterPassCallbackDelegateArray& InOutPassCallbacks, bool bIsPassEnabled) override;
+	virtual void SubscribeToPostProcessingPass(EPostProcessingPass PassId, const FSceneView& View, FAfterPassCallbackDelegateArray& InOutPassCallbacks, bool bIsPassEnabled) override;
 
 	virtual bool IsActiveThisFrame_Internal(const FSceneViewExtensionContext& Context) const override;
 
diff --git a/Engine/Source/Editor/PixelInspector/Private/PixelInspector.cpp b/Engine/Source/Editor/PixelInspector/Private/PixelInspector.cpp
--- a/Engine/Source/Editor/PixelInspector/Private/PixelInspector.cpp
+++ b/Engine/Source/Editor/PixelInspector/Private/PixelInspector.cpp
@@ -857,7 +857,7 @@ namespace PixelInspector
 		FinalColorGamma = (InViewFamily.EngineShowFlags.Tonemapper == 0) ? DEFAULT_DISPLAY_GAMMA : DisplayGamma;
 	}
 
-	void FPixelInspectorSceneViewExtension::SubscribeToPostProcessingPass(EPostProcessingPass PassId, FAfterPassCallbackDelegateArray& InOutPassCallbacks, bool bIsPassEnabled)
+	void FPixelInspectorSceneViewExtension::SubscribeToPostProcessingPass(EPostProcessingPass PassId, const FSceneView& View, FAfterPassCallbackDelegateArray& InOutPassCallbacks, bool bIsPassEnabled)
 	{
 		if (PassId == EPostProcessingPass::FXAA)
 		{
diff --git a/Engine/Source/Editor/PixelInspector/Private/PixelInspector.h b/Engine/Source/Editor/PixelInspector/Private/PixelInspector.h
--- a/Engine/Source/Editor/PixelInspector/Private/PixelInspector.h
+++ b/Engine/Source/Editor/PixelInspector/Private/PixelInspector.h
@@ -35,7 +35,7 @@ namespace PixelInspector
 		virtual void BeginRenderViewFamily(FSceneViewFamily& InViewFamily) override;
 		virtual void PreRenderViewFamily_RenderThread(FRDGBuilder& GraphBuilder, FSceneViewFamily& InViewFamily) override {};
 		virtual void PreRenderView_RenderThread(FRDGBuilder& GraphBuilder, FSceneView& InView) override {};
-		virtual void SubscribeToPostProcessingPass(EPostProcessingPass PassId, FAfterPassCallbackDelegateArray& InOutPassCallbacks, bool bIsPassEnabled) override;
+		virtual void SubscribeToPostProcessingPass(EPostProcessingPass PassId, const FSceneView& View, FAfterPassCallbackDelegateArray& InOutPassCallbacks, bool bIsPassEnabled) override;
 
 	public:
 		FScreenPassTexture PostProcessPassAfterFxaa_RenderThread(FRDGBuilder& GraphBuilder, const FSceneView& View, const FPostProcessMaterialInputs& InOutInputs);
diff --git a/Engine/Source/Runtime/Engine/Public/SceneViewExtension.h b/Engine/Source/Runtime/Engine/Public/SceneViewExtension.h
--- a/Engine/Source/Runtime/Engine/Public/SceneViewExtension.h
+++ b/Engine/Source/Runtime/Engine/Public/SceneViewExtension.h
@@ -185,7 +185,10 @@ public:
 	*  - The pass MUST write to the override output texture if it is active (this occurs when the pass is the last in the post processing chain writing to the back buffer).
 	*    For performance reasons it is recommended to only subscribe to a pass when the pass will produce a GPU resource. Calling 
 	*/
-	virtual void SubscribeToPostProcessingPass(EPostProcessingPass Pass, FAfterPassCallbackDelegateArray& InOutPassCallbacks, bool bIsPassEnabled) {};
+	virtual void SubscribeToPostProcessingPass(EPostProcessingPass Pass, const FSceneView& View, FAfterPassCallbackDelegateArray& InOutPassCallbacks, bool bIsPassEnabled) {}
+
+	UE_DEPRECATED(5.5, "SubscribeToPostProcessingPass now takes a SceneView")
+	virtual void SubscribeToPostProcessingPass(EPostProcessingPass Pass, FAfterPassCallbackDelegateArray& InOutPassCallbacks, bool bIsPassEnabled) {}
 
 	/**
 	 * Allows to render content after the 3D content scene, useful for debugging
diff --git a/Engine/Source/Runtime/Renderer/Private/PostProcess/PostProcessing.cpp b/Engine/Source/Runtime/Renderer/Private/PostProcess/PostProcessing.cpp
--- a/Engine/Source/Runtime/Renderer/Private/PostProcess/PostProcessing.cpp
+++ b/Engine/Source/Runtime/Renderer/Private/PostProcess/PostProcessing.cpp
@@ -676,7 +676,7 @@ void AddPostProcessingPasses(
 
 		for (int32 ViewExt = 0; ViewExt < View.Family->ViewExtensions.Num(); ++ViewExt)
 		{
-			View.Family->ViewExtensions[ViewExt]->SubscribeToPostProcessingPass(ISceneViewExtension::EPostProcessingPass::SSRInput, SSRInputDelegates, true);
+			View.Family->ViewExtensions[ViewExt]->SubscribeToPostProcessingPass(ISceneViewExtension::EPostProcessingPass::SSRInput, View, SSRInputDelegates, true);
 
 			for (int32 SceneViewPassId = static_cast<int32>(ISceneViewExtension::EPostProcessingPass::MotionBlur); SceneViewPassId != static_cast<int32>(ISceneViewExtension::EPostProcessingPass::MAX); SceneViewPassId++)
 			{
@@ -685,6 +685,7 @@ void AddPostProcessingPasses(
 
 				View.Family->ViewExtensions[ViewExt]->SubscribeToPostProcessingPass(
 					SceneViewPass,
+					View,
 					PassSequence.GetAfterPassCallbacks(PostProcessingPass),
 					PassSequence.IsEnabled(PostProcessingPass));
 			}
