PPMChainGraph: Export public classes & structs, per third-party request.

diff --git a/Engine/Plugins/Experimental/PostProcessMaterialChainGraph/Source/PostProcessMaterialChainGraph/Public/PPMChainGraph.h b/Engine/Plugins/Experimental/PostProcessMaterialChainGraph/Source/PostProcessMaterialChainGraph/Public/PPMChainGraph.h
--- a/Engine/Plugins/Experimental/PostProcessMaterialChainGraph/Source/PostProcessMaterialChainGraph/Public/PPMChainGraph.h
+++ b/Engine/Plugins/Experimental/PostProcessMaterialChainGraph/Source/PostProcessMaterialChainGraph/Public/PPMChainGraph.h
@@ -63,7 +63,7 @@ enum class EPPMChainGraphExecutionLocation : uint8
 * This struct is used for customizing Input and External Texture input selection.
 */
 USTRUCT(BlueprintType)
-struct FPPMChainGraphInput
+struct PPMCHAINGRAPH_API FPPMChainGraphInput
 {
 	GENERATED_USTRUCT_BODY()
 
@@ -73,7 +73,7 @@ struct FPPMChainGraphInput
 
 
 USTRUCT(Blueprintable)
-struct FPPMChainGraphPostProcessPass
+struct PPMCHAINGRAPH_API FPPMChainGraphPostProcessPass
 {
 	GENERATED_USTRUCT_BODY()
 
diff --git a/Engine/Plugins/Experimental/PostProcessMaterialChainGraph/Source/PostProcessMaterialChainGraph/Public/PPMChainGraphComponent.h b/Engine/Plugins/Experimental/PostProcessMaterialChainGraph/Source/PostProcessMaterialChainGraph/Public/PPMChainGraphComponent.h
--- a/Engine/Plugins/Experimental/PostProcessMaterialChainGraph/Source/PostProcessMaterialChainGraph/Public/PPMChainGraphComponent.h
+++ b/Engine/Plugins/Experimental/PostProcessMaterialChainGraph/Source/PostProcessMaterialChainGraph/Public/PPMChainGraphComponent.h
@@ -29,7 +29,7 @@ enum class ECameraViewHandling : uint8
 };
 
 UCLASS(hidecategories = (Transform, Collision, Object, Physics, SceneComponent, PostProcessVolume, Projection, Rendering, PlanarReflection), ClassGroup = Rendering, editinlinenew, meta = (BlueprintSpawnableComponent, DisplayName = "Post Process Material Chain Graph Executor Component"))
-class UPPMChainGraphExecutorComponent : public USceneComponent
+class PPMCHAINGRAPH_API UPPMChainGraphExecutorComponent : public USceneComponent
 {
 	GENERATED_UCLASS_BODY()
 public:
